MP4Utilities: Created a new utility plugin to contain helpers for mp4 files. This first version includes an mp4 file multiplexer.

#rnx

#Synced-CL 49644370

diff --git a/Engine/Plugins/Media/MP4Utilities/MP4Utilities.uplugin b/Engine/Plugins/Media/MP4Utilities/MP4Utilities.uplugin
new file mode 100644
--- /dev/null
+++ b/Engine/Plugins/Media/MP4Utilities/MP4Utilities.uplugin
@@ -0,0 +1,21 @@
+{
+  "FileVersion": 3,
+  "Version": 1,
+  "VersionName": "1.0",
+  "FriendlyName": "MP4 (ISO/IEC 14496-12) utilities",
+  "Description": "Provides helpers to work with mp4 files",
+  "Category": "Media",
+  "CreatedBy": "Epic Games, Inc.",
+  "CreatedByURL": "https://epicgames.com",
+  "EnabledByDefault": false,
+  "CanContainContent": false,
+  "IsBetaVersion" : true,
+  "Installed": false,
+  "Modules": [
+    {
+      "Name": "MP4Muxer",
+      "Type": "Runtime",
+      "LoadingPhase": "Default"
+    }
+  ]
+}
\ No newline at end of file
diff --git a/Engine/Plugins/Media/MP4Utilities/Source/MP4Muxer/MP4Muxer.Build.cs b/Engine/Plugins/Media/MP4Utilities/Source/MP4Muxer/MP4Muxer.Build.cs
new file mode 100644
--- /dev/null
+++ b/Engine/Plugins/Media/MP4Utilities/Source/MP4Muxer/MP4Muxer.Build.cs
@@ -0,0 +1,17 @@
+// Copyright Epic Games, Inc. All Rights Reserved.
+
+using EpicGames.Core;
+using System.IO;
+
+namespace UnrealBuildTool.Rules
+{
+	public class MP4Muxer : ModuleRules
+	{
+		public MP4Muxer(ReadOnlyTargetRules Target) : base(Target)
+		{
+			PrivateDependencyModuleNames.AddRange(new string[] {
+				"Core"
+            });
+		}
+	}
+}
diff --git a/Engine/Plugins/Media/MP4Utilities/Source/MP4Muxer/Private/MP4Muxer.cpp b/Engine/Plugins/Media/MP4Utilities/Source/MP4Muxer/Private/MP4Muxer.cpp
new file mode 100644
--- /dev/null
+++ b/Engine/Plugins/Media/MP4Utilities/Source/MP4Muxer/Private/MP4Muxer.cpp
@@ -0,0 +1,1262 @@
+// Copyright Epic Games, Inc. All Rights Reserved.
+
+#include "MP4Muxer.h"
+
+#include "GenericPlatform/GenericPlatformFile.h"
+#include "HAL/FileManager.h"
+#include "HAL/PlatformFileManager.h"
+#include "HAL/PlatformAffinity.h"
+#include "HAL/Runnable.h"
+#include "HAL/RunnableThread.h"
+#include "HAL/Event.h"
+#include "HAL/PlatformProcess.h"
+#include "HAL/CriticalSection.h"
+#include "Templates/UniquePtr.h"
+#include "Containers/Queue.h"
+#include "Misc/DateTime.h"
+#include "Math/BigInt.h"
+
+#include "MP4MuxerModule.h"
+#include "MP4MuxerBoxes.h"
+
+#include <numeric>
+
+
+namespace TimeConversion
+{
+static int64 ToTimescale(uint32 InTargetDenominator, int64 InSourceNumerator, uint32 InSourceDenominator)
+{
+	if (InTargetDenominator == InSourceDenominator)
+	{
+		return InSourceNumerator;
+	}
+	else if (InSourceNumerator == 0)
+	{
+		return 0;
+	}
+	else if (InSourceDenominator == 0 || InTargetDenominator == 0)
+	{
+		return InSourceNumerator >= 0 ? 0x7fffffffffffffffLL : -0x7fffffffffffffffLL;
+	}
+
+	bool bIsNeg = InSourceNumerator < 0;
+	TBigInt<128> n(bIsNeg ? -InSourceNumerator : InSourceNumerator);
+	TBigInt<128> d(InSourceDenominator);
+	TBigInt<128> s(InTargetDenominator);
+
+	n *= s;
+	n /= d;
+
+	int64 r = n.ToInt();
+	return bIsNeg ? -r : r;
+}
+}
+
+
+class FMP4RawMuxer : public IMP4RawMuxer, public FRunnable
+{
+public:
+	virtual ~FMP4RawMuxer();
+
+	bool Configure(const FConfiguration& InConfiguration) override;
+	int32 AddTrack(const FTrackSpec& InTrackSpec) override;
+	bool Start(FRequestTrackDataDelegate InSampleRequestDelegate, FStatusDelegate InStatusDelegate) override;
+	bool AddTrackSample(int32 InTrack, const FTrackSample& InTrackSample) override;
+	bool StopAndClose() override;
+	FString GetLastError() override
+	{ return LastError; }
+private:
+	struct FTrackMuxBoxes
+	{
+		MP4MuxerBoxes::FBoxTREX* TREX = nullptr;
+		MP4MuxerBoxes::FBoxTKHD* TKHD = nullptr;
+		MP4MuxerBoxes::FBoxMDHD* MDHD = nullptr;
+		MP4MuxerBoxes::FBoxSTTS* STTS = nullptr;
+		MP4MuxerBoxes::FBoxSTSS* STSS = nullptr;
+		MP4MuxerBoxes::FBoxCTTS* CTTS = nullptr;
+		MP4MuxerBoxes::FBoxSTSC* STSC = nullptr;
+		MP4MuxerBoxes::FBoxSTSZ* STSZ = nullptr;
+		MP4MuxerBoxes::FBoxSTCO* STCO = nullptr;
+		MP4MuxerBoxes::FBoxSUBS* SUBS = nullptr;
+		MP4MuxerBoxes::FBoxELST* ELST = nullptr;
+	};
+
+	struct FTrackChunkSamples
+	{
+		void Empty()
+		{
+			Samples.Empty();
+			// Subtract the wanted interleave duration from the amount accumulated to account for
+			// the gathered samples having a longer duration than asked for. We need to keep the
+			// remainder so the next run may collect one (or however many) fewer samples to prevent
+			// the individual track chunks from not drifting apart.
+			AccumulatedDuration -= InterleaveDurationNeeded;
+			SumOfSampleDurations = 0;
+		}
+		struct FSample
+		{
+			TArray<FTrackSample::FSubSampleInfo> SubSamples;
+			TArray64<uint8> Data;
+			int64 DTS;
+			int64 PTS;
+			uint32 Duration;
+			bool bIsKeyframe;
+		};
+		FTrackMuxBoxes TrackMuxBoxes;
+		TArray<TUniquePtr<FSample>> Samples;
+		int64 TotalTrackDuration = 0;
+		int64 AccumulatedDuration = 0;
+		int64 SumOfSampleDurations = 0;
+		int64 InterleaveDurationNeeded = 0;
+		uint32 Timescale = 0;
+		uint32 NextSampleNumToRequest = 0;
+		bool bGotLastSample = false;
+		bool bHasFinished = false;
+	};
+
+	bool SetError(const FString& InErrorMsg)
+	{
+		FScopeLock lock(&Lock);
+		LastError = InErrorMsg;
+		bFailed = true;
+		return false;
+	}
+	bool SetWriteError()
+	{
+		return SetError(FString::Printf(TEXT("Could not write to file \"%s\"."), *Configuration.OutputFilename));
+	}
+	bool SetUpdateError(const TCHAR* InWhat, const TCHAR* InBox)
+	{
+		return SetError(FString::Printf(TEXT("Failed to update %s in `%s` box, value now needs 64 bits!"), InWhat, InBox));
+	}
+	uint32 Run() override;
+	void CloseFile();
+	bool RequestSampleData();
+	bool WriteFileStart();
+	bool WriteAccumulatedSamplesStandard();
+	bool WriteAccumulatedSamplesFragmented();
+	bool WriteAccumulatedSamples();
+	bool WriteFileEnd();
+	bool WriteUpdate(int64 InFilePosition, TConstArrayView<uint8> InData);
+
+	FConfiguration Configuration;
+	TArray<FTrackSpec> TrackSpecs;
+	FRequestTrackDataDelegate SampleRequestDelegate;
+	FStatusDelegate StatusDelegate;
+	FString LastError;
+	bool bFailed = false;
+	bool bTerminate = false;
+	FRunnableThread* Thread = nullptr;
+	FEvent* WorkerEvent = nullptr;
+	IFileHandle* FileHandle = nullptr;
+	int64 FileOffsetOfStandardMDAT = -1;
+	uint32 FragmentSequenceNumber = 0;
+	bool bWroteMOOV = false;
+
+	FCriticalSection Lock;
+
+	TArray<FTrackChunkSamples> TrackChunkSamples;
+	int32 MinimumMajorISOBrand = 0;
+	int32 WrittenMajorISOBrand = 0;
+	TUniquePtr<MP4MuxerBoxes::FBoxFTYP> FTYP;
+	TUniquePtr<MP4MuxerBoxes::FBoxMOOV> MOOV;
+	MP4MuxerBoxes::FBoxMVHD* MVHD = nullptr;
+	MP4MuxerBoxes::FBoxMEHD* MEHD = nullptr;
+
+	TArray<FSampleRequest> TracksAwaitingSamples;
+	MP4MuxerBoxes::FBoxBase::FRewriteBoxChangeDelegate RewriteChangeDelegate;
+};
+
+TSharedRef<IMP4RawMuxer, ESPMode::ThreadSafe> IMP4RawMuxer::Create()
+{
+	return MakeShared<FMP4RawMuxer, ESPMode::ThreadSafe>();
+}
+
+TArray<uint8> IMP4RawMuxer::WrapDataInBox(uint32 InBoxAtom, TConstArrayView<uint8> InBoxData)
+{
+	TArray<uint8> Wrap;
+	Wrap.SetNumUninitialized(InBoxData.Num() + 8);
+	uint32* bd = reinterpret_cast<uint32*>(Wrap.GetData());
+	*bd++ = MP4MuxerBoxes::EndianSwap((uint32)(InBoxData.Num() + 8));
+	*bd++ = MP4MuxerBoxes::EndianSwap(InBoxAtom);
+	FMemory::Memcpy(bd, InBoxData.GetData(), InBoxData.Num());
+	return Wrap;
+}
+
+TArray<uint8> IMP4RawMuxer::WrapDataInBox(uint32 InBoxAtom, uint8 InBoxVersion, uint32 InBoxFlags, TConstArrayView<uint8> InBoxData)
+{
+	TArray<uint8> Wrap;
+	Wrap.SetNumUninitialized(InBoxData.Num() + 12);
+	uint32* bd = reinterpret_cast<uint32*>(Wrap.GetData());
+	*bd++ = MP4MuxerBoxes::EndianSwap((uint32)(InBoxData.Num() + 12));
+	*bd++ = MP4MuxerBoxes::EndianSwap(InBoxAtom);
+	*bd++ = MP4MuxerBoxes::EndianSwap((uint32)((InBoxFlags & 0x00ffffff) | ((uint32)InBoxVersion << 24)));
+	FMemory::Memcpy(bd, InBoxData.GetData(), InBoxData.Num());
+	return Wrap;
+}
+
+
+
+FMP4RawMuxer::~FMP4RawMuxer()
+{
+	StopAndClose();
+}
+
+void FMP4RawMuxer::CloseFile()
+{
+	if (FileHandle)
+	{
+		delete FileHandle;
+		FileHandle = nullptr;
+	}
+}
+
+bool FMP4RawMuxer::Configure(const FConfiguration& InConfiguration)
+{
+	Configuration = InConfiguration;
+	if (Configuration.InterleaveDuration.GetTicks() < 0)
+	{
+		return SetError(FString::Printf(TEXT("Invalid interleave duration of %.4f msec."), Configuration.InterleaveDuration.GetTotalSeconds()));
+	}
+	if (Configuration.InterleaveDuration.GetTicks() == 0)
+	{
+		Configuration.InterleaveDuration = FTimespan::FromMilliseconds(500);
+	}
+
+	RewriteChangeDelegate.BindRaw(this, &FMP4RawMuxer::WriteUpdate);
+	return true;
+}
+
+int32 FMP4RawMuxer::AddTrack(const FTrackSpec& InTrackSpec)
+{
+	// Some basic checks
+	if (!InTrackSpec.SampleEntryFormat)
+	{
+		SetError(TEXT("Unset sample entry format"));
+		return -1;
+	}
+	if (!InTrackSpec.Timescale)
+	{
+		SetError(TEXT("Unset time scale"));
+		return -1;
+	}
+
+	TrackSpecs.Emplace(InTrackSpec);
+	return TrackSpecs.Num() - 1;
+}
+
+
+
+
+bool FMP4RawMuxer::Start(FRequestTrackDataDelegate InSampleRequestDelegate, FStatusDelegate InStatusDelegate)
+{
+	// If already in error just return.
+	if (LastError.Len())
+	{
+		return false;
+	}
+	// Need to have a sample retrieval delegate.
+	if (!InSampleRequestDelegate.IsBound() || !InStatusDelegate.IsBound())
+	{
+		return SetError(TEXT("Unbound delegate passed to Start()"));
+	}
+	// Don't call Start() twice.
+	if (SampleRequestDelegate.IsBound())
+	{
+		return SetError(TEXT("Start() has already been called before"));
+	}
+	SampleRequestDelegate = InSampleRequestDelegate;
+	StatusDelegate = InStatusDelegate;
+
+	// Set up the per track sample chunks
+	TrackChunkSamples.SetNum(TrackSpecs.Num());
+	for(int32 i=0; i<TrackChunkSamples.Num(); ++i)
+	{
+		TrackChunkSamples[i].Timescale = TrackSpecs[i].Timescale;
+		TrackChunkSamples[i].InterleaveDurationNeeded = TimeConversion::ToTimescale(TrackSpecs[i].Timescale, Configuration.InterleaveDuration.GetTicks(), (uint32)ETimespan::TicksPerSecond);
+		if (!TrackChunkSamples[i].InterleaveDurationNeeded)
+		{
+			return SetError(FString::Printf(TEXT("Interleave duration not representable by a timescale of %u"), TrackSpecs[i].Timescale));
+		}
+	}
+
+	WorkerEvent = FPlatformProcess::GetSynchEventFromPool(true);
+	if (!WorkerEvent)
+	{
+		return SetError(TEXT("Start() failed to create worker event"));
+	}
+	Thread = FRunnableThread::Create(this, TEXT("MP4Muxer"));
+	if (!Thread)
+	{
+		return SetError(TEXT("Start() failed to create worker thread"));
+	}
+	return true;
+}
+
+
+bool FMP4RawMuxer::StopAndClose()
+{
+	bTerminate = true;
+	if (Thread)
+	{
+		Thread->WaitForCompletion();
+		delete Thread;
+		Thread = nullptr;
+	}
+	if (WorkerEvent)
+	{
+		FPlatformProcess::ReturnSynchEventToPool(WorkerEvent);
+		WorkerEvent = nullptr;
+	}
+	return true;
+}
+
+bool FMP4RawMuxer::WriteFileStart()
+{
+	// Write the starting `ftyp` box and a size-reserving `free` box followed by the `mdat` when writing a standard file.
+	MP4MuxerBoxes::FBoxBase* FileLevelBoxes = new MP4MuxerBoxes::FBoxBase(0);
+	FTYP = MakeUnique<MP4MuxerBoxes::FBoxFTYP>();
+	WrittenMajorISOBrand = 2;
+	FTYP->SetMajorBrand(MP4MuxerBoxes::BoxAtom('i','s','o','2'));
+	FTYP->SetMinorVersion(1);
+	FTYP->AddCompatibleBrand(MP4MuxerBoxes::BoxAtom('a','v','c','1'));	// `iso2` (or better) requires this
+	FTYP->AddCompatibleBrand(MP4MuxerBoxes::BoxAtom('m','p','4','2'));
+	FileLevelBoxes->AddChild(FTYP.Get());
+	MinimumMajorISOBrand = WrittenMajorISOBrand;
+	MP4MuxerBoxes::FBoxMDAT* MDAT = nullptr;
+	if (Configuration.MuxMode == IMP4RawMuxer::FConfiguration::EMuxMode::Standard)
+	{
+		// We may need additional 8 bytes if the following `mdat` exceeds 4GiB,
+		// so we insert a `free` box with no additional zero bytes as this box size is 8 bytes.
+		MP4MuxerBoxes::FBoxFREE* FREE = new MP4MuxerBoxes::FBoxFREE;
+		FREE->SetNumberOfEmptyBytes(0);
+		FileLevelBoxes->AddChild(FREE);
+		MDAT = new MP4MuxerBoxes::FBoxMDAT;
+		FileLevelBoxes->AddChild(MDAT);
+	}
+	// Build the binary representation of the boxes
+	MP4MuxerBoxes::FBoxBuilderData builderData;
+	FileLevelBoxes->BuildData(builderData);
+	FileLevelBoxes->CalculateBoxSizes();
+	bool bWriteOk = FileLevelBoxes->VisitForWriting([&](MP4MuxerBoxes::FBoxBase* InBox)
+	{
+		InBox->SetAbsoluteFileOffset(FileHandle->Tell());
+		return InBox->GetCompiledBoxData().Num() ? FileHandle->Write(InBox->GetCompiledBoxData().GetData(), InBox->GetCompiledBoxData().Num()) : true;
+	});
+	if (!bWriteOk)
+	{
+		return SetWriteError();
+	}
+	FileOffsetOfStandardMDAT = MDAT ? MDAT->GetAbsoluteFileOffset() : -1;
+	// We need to keep the ftyp box!
+	FileLevelBoxes->RemoveChild(FTYP.Get());
+	delete FileLevelBoxes;
+	FileLevelBoxes = nullptr;
+
+	// Create the `moov` box and the track structures.
+	// We need these when writing standard as well as fragmented files, so we can create them here.
+	MOOV = MakeUnique<MP4MuxerBoxes::FBoxMOOV>();
+	MVHD = new MP4MuxerBoxes::FBoxMVHD;
+	MP4MuxerBoxes::FBoxMVEX* MVEX = Configuration.MuxMode == IMP4RawMuxer::FConfiguration::EMuxMode::Fragmented ? new MP4MuxerBoxes::FBoxMVEX : nullptr;
+	MOOV->AddChild(MVHD);
+	uint64 UtcCreationDate(static_cast<uint64>((FDateTime::UtcNow() - FDateTime(1904, 1, 1)).GetTotalSeconds()));
+	MVHD->SetCreationTime(UtcCreationDate);
+	MVHD->SetModificationTime(UtcCreationDate);
+	// Start out with a millisecond timescale for the movie box.
+	MVHD->SetTimescale(1000);
+	MVHD->SetDuration(0);
+	MVHD->SetNextTrackID(TrackSpecs.Num() + 1);
+	if (MVEX)
+	{
+		MOOV->AddChild(MVEX);
+		/*
+			We patch the `mvhd` duration at the end so there is no need to write an `mehd` box
+
+			MEHD = new MP4MuxerBoxes::FBoxMEHD;
+			MVEX->AddChild(MEHD);
+		*/
+	}
+
+	for(int32 TrackIdx=0; TrackIdx<TrackSpecs.Num(); ++TrackIdx)
+	{
+		bool bIsVideo = false;
+		bool bIsAudio = false;
+
+		MP4MuxerBoxes::FBoxTRAK* TRAK = new MP4MuxerBoxes::FBoxTRAK;
+		MOOV->AddChild(TRAK);
+		MP4MuxerBoxes::FBoxTKHD* TKHD = new MP4MuxerBoxes::FBoxTKHD;
+		TrackChunkSamples[TrackIdx].TrackMuxBoxes.TKHD = TKHD;
+		TRAK->AddChild(TKHD);
+		TKHD->SetCreationTime(UtcCreationDate);
+		TKHD->SetModificationTime(UtcCreationDate);
+		TKHD->SetDuration(0);
+		TKHD->SetTrackID(1U + TrackIdx);
+		TKHD->SetTrackEnabled(TrackSpecs[TrackIdx].bEnabled ? 1 : 0);
+		TKHD->SetTrackInMovie(TrackSpecs[TrackIdx].bInMovie ? 1 : 0);
+		TKHD->SetTrackInPreview(0);
+		TKHD->SetLayer(TrackSpecs[TrackIdx].Layer);
+		TKHD->SetAlternateGroup(TrackSpecs[TrackIdx].AlternateGroup);
+
+		if (TrackSpecs[TrackIdx].Properties.IsType<FTrackSpec::FVideo>())
+		{
+			auto const& vid = TrackSpecs[TrackIdx].Properties.Get<FTrackSpec::FVideo>();
+			TKHD->SetWidth((uint16)vid.DisplayWidth);
+			TKHD->SetHeight((uint16)vid.DisplayHeight);
+			TKHD->SetVolume(0);
+			bIsVideo = true;
+		}
+		else if (TrackSpecs[TrackIdx].Properties.IsType<FTrackSpec::FAudio>())
+		{
+			auto const& vid = TrackSpecs[TrackIdx].Properties.Get<FTrackSpec::FAudio>();
+			TKHD->SetVolume(1);
+			bIsAudio = true;
+		}
+
+		// TREX ?
+		if (MVEX)
+		{
+			MP4MuxerBoxes::FBoxTREX* TREX = new MP4MuxerBoxes::FBoxTREX;
+			TrackChunkSamples[TrackIdx].TrackMuxBoxes.TREX = TREX;
+			MVEX->AddChild(TREX);
+			TREX->SetTrackID(1U + TrackIdx);
+			TREX->SetDefaultSampleDuration(TrackSpecs[TrackIdx].ConstantSampleDuration);
+			// Constant size not used yet.
+			TREX->SetDefaultSampleSize(0);	// TrackSpecs[TrackIdx].ConstantSampleSize);
+			if (TrackSpecs[TrackIdx].bIsAllKeyframes)
+			{
+				TREX->SetDefaultSampleFlags(0x0);
+			}
+			else
+			{
+				// sample_is_non_sync_sample = 1
+				TREX->SetDefaultSampleFlags(0x10000);
+			}
+		}
+
+		// MDIA
+		MP4MuxerBoxes::FBoxMDIA* MDIA = new MP4MuxerBoxes::FBoxMDIA;
+		TRAK->AddChild(MDIA);
+		// MDHD
+		MP4MuxerBoxes::FBoxMDHD* MDHD = new MP4MuxerBoxes::FBoxMDHD;
+		TrackChunkSamples[TrackIdx].TrackMuxBoxes.MDHD = MDHD;
+		MDIA->AddChild(MDHD);
+		MDHD->SetCreationTime(UtcCreationDate);
+		MDHD->SetModificationTime(UtcCreationDate);
+		MDHD->SetDuration(0);
+		MDHD->SetTimescale(TrackSpecs[TrackIdx].Timescale);
+		MDHD->SetLanguage(TrackSpecs[TrackIdx].Language[0], TrackSpecs[TrackIdx].Language[1], TrackSpecs[TrackIdx].Language[2]);
+
+		// HDLR
+		MP4MuxerBoxes::FBoxHDLR* HDLR = new MP4MuxerBoxes::FBoxHDLR;
+		MDIA->AddChild(HDLR);
+		// MINF
+		MP4MuxerBoxes::FBoxMINF* MINF = new MP4MuxerBoxes::FBoxMINF;
+		MDIA->AddChild(MINF);
+		if (bIsVideo)
+		{
+			HDLR->SetHandlerType(MP4MuxerBoxes::BoxAtom('v','i','d','e'));
+			HDLR->SetName(TEXT("VideoHandler"));
+			MP4MuxerBoxes::FBoxVMHD* VMHD = new MP4MuxerBoxes::FBoxVMHD;
+			MINF->AddChild(VMHD);
+		}
+		else if (bIsAudio)
+		{
+			HDLR->SetHandlerType(MP4MuxerBoxes::BoxAtom('s','o','u','n'));
+			HDLR->SetName(TEXT("SoundHandler"));
+			MP4MuxerBoxes::FBoxSMHD* SMHD = new MP4MuxerBoxes::FBoxSMHD;
+			MINF->AddChild(SMHD);
+		}
+		else
+		{
+			unimplemented();
+			return false;
+		}
+
+		// DINF/DREF/URL
+		MP4MuxerBoxes::FBoxDINF* DINF = new MP4MuxerBoxes::FBoxDINF;
+		MINF->AddChild(DINF);
+		MP4MuxerBoxes::FBoxDREF* DREF = new MP4MuxerBoxes::FBoxDREF;
+		DINF->AddChild(DREF);
+		MP4MuxerBoxes::FBoxURL* URL_ = new MP4MuxerBoxes::FBoxURL;
+		DREF->AddChild(URL_);
+
+		// STBL
+		MP4MuxerBoxes::FBoxSTBL* STBL = new MP4MuxerBoxes::FBoxSTBL;
+		MINF->AddChild(STBL);
+		MP4MuxerBoxes::FBoxSTSD* STSD = new MP4MuxerBoxes::FBoxSTSD;
+		STBL->AddChild(STSD);
+		if (bIsVideo)
+		{
+			auto const& vid = TrackSpecs[TrackIdx].Properties.Get<FTrackSpec::FVideo>();
+			MP4MuxerBoxes::FBoxVisualSampleEntry* Visual = new MP4MuxerBoxes::FBoxVisualSampleEntry(TrackSpecs[TrackIdx].SampleEntryFormat);
+			Visual->SetWidth((uint16)vid.DisplayWidth);
+			Visual->SetHeight((uint16)vid.DisplayHeight);
+			Visual->SetCompressorName(vid.CompressorName);
+			Visual->SetAdditionalBoxes(TrackSpecs[TrackIdx].SampleEntryBoxes);
+			STSD->AddChild(Visual);
+		}
+		else if (bIsAudio)
+		{
+			auto const& aud = TrackSpecs[TrackIdx].Properties.Get<FTrackSpec::FAudio>();
+			MP4MuxerBoxes::FBoxAudioSampleEntry* Audio = new MP4MuxerBoxes::FBoxAudioSampleEntry(TrackSpecs[TrackIdx].SampleEntryFormat);
+			if (Audio->SetSampleRate(aud.SamplingRate))
+			{
+				STSD->SetVersion(1);
+			}
+			Audio->SetChannelCount(aud.NumberOfChannels);
+			Audio->SetAdditionalBoxes(TrackSpecs[TrackIdx].SampleEntryBoxes);
+			STSD->AddChild(Audio);
+		}
+		else
+		{
+			unimplemented();
+			return false;
+		}
+
+		MP4MuxerBoxes::FBoxSTTS* STTS = new MP4MuxerBoxes::FBoxSTTS;
+		TrackChunkSamples[TrackIdx].TrackMuxBoxes.STTS = STTS;
+		STBL->AddChild(STTS);
+		if (Configuration.MuxMode == IMP4RawMuxer::FConfiguration::EMuxMode::Standard)
+		{
+			if (!TrackSpecs[TrackIdx].bIsAllKeyframes)
+			{
+				// Sync sample box (not required when every sample is a sync sample or when writing fragmented file)
+				MP4MuxerBoxes::FBoxSTSS* STSS = new MP4MuxerBoxes::FBoxSTSS;
+				TrackChunkSamples[TrackIdx].TrackMuxBoxes.STSS = STSS;
+				STBL->AddChild(STSS);
+			}
+			// Composition time offset box (should not be required when every sample is a sync sample though)
+			MP4MuxerBoxes::FBoxCTTS* CTTS = new MP4MuxerBoxes::FBoxCTTS;
+			TrackChunkSamples[TrackIdx].TrackMuxBoxes.CTTS = CTTS;
+			STBL->AddChild(CTTS);
+		}
+		MP4MuxerBoxes::FBoxSTSC* STSC = new MP4MuxerBoxes::FBoxSTSC;
+		TrackChunkSamples[TrackIdx].TrackMuxBoxes.STSC = STSC;
+		STBL->AddChild(STSC);
+		MP4MuxerBoxes::FBoxSTSZ* STSZ = new MP4MuxerBoxes::FBoxSTSZ;
+		TrackChunkSamples[TrackIdx].TrackMuxBoxes.STSZ = STSZ;
+		STBL->AddChild(STSZ);
+		MP4MuxerBoxes::FBoxSTCO* STCO = new MP4MuxerBoxes::FBoxSTCO;
+		TrackChunkSamples[TrackIdx].TrackMuxBoxes.STCO = STCO;
+		STBL->AddChild(STCO);
+		// Sub-sample information? These go into the `stbl` for standard files and into the `traf` for fragmented files.
+		if (Configuration.MuxMode == IMP4RawMuxer::FConfiguration::EMuxMode::Standard && TrackSpecs[TrackIdx].SubSampleFlags.IsSet())
+		{
+			MP4MuxerBoxes::FBoxSUBS* SUBS = new MP4MuxerBoxes::FBoxSUBS;
+			TrackChunkSamples[TrackIdx].TrackMuxBoxes.SUBS = SUBS;
+			STBL->AddChild(SUBS);
+			SUBS->SetFlags(TrackSpecs[TrackIdx].SubSampleFlags.GetValue());
+		}
+	}
+	return true;
+}
+
+
+bool FMP4RawMuxer::WriteUpdate(int64 InFilePosition, TConstArrayView<uint8> InData)
+{
+	if (!FileHandle)
+	{
+		return false;
+	}
+	int64 PosNow = FileHandle->Tell();
+	return FileHandle->Seek(InFilePosition) && FileHandle->Write(InData.GetData(), InData.Num()) && FileHandle->Seek(PosNow);
+}
+
+
+bool FMP4RawMuxer::WriteFileEnd()
+{
+	if (!FTYP || !MOOV || !MVHD)
+	{
+		return false;
+	}
+	int64 FilePosAtEnd = FileHandle->Tell();
+
+	// Calculate a common track time scale if possible to use in the movie boxes.
+	uint32 MovieTimescale = TrackChunkSamples.Num() ? TrackChunkSamples[0].Timescale : MVHD->GetTimescale();
+	uint64 LcmTimescale = MovieTimescale;
+	for(int32 i=1,iMax=TrackChunkSamples.Num(); i<iMax; ++i)
+	{
+		LcmTimescale = std::lcm((uint64)LcmTimescale, (uint64)TrackChunkSamples[i].Timescale);
+		// Excessivily large such that we cannot represent a minimum of 7200 seconds without overflowing a 32 bit value?
+		if (LcmTimescale > 0x80000)
+		{
+			MovieTimescale = MVHD->GetTimescale();
+			break;
+		}
+		MovieTimescale = (uint32)LcmTimescale;
+	}
+
+	if (MovieTimescale != MVHD->GetTimescale())
+	{
+		if (!MVHD->UpdateTimescale(MovieTimescale))
+		{
+			return SetUpdateError(TEXT("timescale"), TEXT("mvhd"));
+		}
+	}
+
+	// Update the track durations in the `tkhd` and `mdhd` boxes as well as the movie duration in `mvhd` and `mehd`
+	int64 LongestTrackInMovieUnits = 0;
+	for(auto &tcs : TrackChunkSamples)
+	{
+		int64 TrackDurationInMovieUnits = TimeConversion::ToTimescale(MovieTimescale, tcs.TotalTrackDuration, tcs.Timescale);
+		if (TrackDurationInMovieUnits > LongestTrackInMovieUnits)
+		{
+			LongestTrackInMovieUnits = TrackDurationInMovieUnits;
+		}
+		if (!tcs.TrackMuxBoxes.TKHD->UpdateDuration(TrackDurationInMovieUnits))
+		{
+			return SetUpdateError(TEXT("duration"), TEXT("tkhd"));
+		}
+		if (!tcs.TrackMuxBoxes.TKHD->RewriteChanges(RewriteChangeDelegate))
+		{
+			return SetWriteError();
+		}
+		if (!tcs.TrackMuxBoxes.MDHD->UpdateDuration(tcs.TotalTrackDuration))
+		{
+			return SetUpdateError(TEXT("duration"), TEXT("mdhd"));
+		}
+		if (!tcs.TrackMuxBoxes.MDHD->RewriteChanges(RewriteChangeDelegate))
+		{
+			return SetWriteError();
+		}
+		if (tcs.TrackMuxBoxes.ELST)
+		{
+			if (!tcs.TrackMuxBoxes.ELST->UpdateEditDuration(tcs.TotalTrackDuration))
+			{
+				return SetUpdateError(TEXT("edit duration"), TEXT("elst"));
+			}
+			if (!tcs.TrackMuxBoxes.ELST->RewriteChanges(RewriteChangeDelegate))
+			{
+				return SetWriteError();
+			}
+		}
+	}
+	if (!MVHD->UpdateDuration(LongestTrackInMovieUnits))
+	{
+		return SetUpdateError(TEXT("duration"), TEXT("mvhd"));
+	}
+	if (!MVHD->RewriteChanges(RewriteChangeDelegate))
+	{
+		return SetWriteError();
+	}
+	if (MEHD)
+	{
+		if (!MEHD->UpdateFragmentDuration(LongestTrackInMovieUnits))
+		{
+			return SetUpdateError(TEXT("fragment duration"), TEXT("mehd"));
+		}
+		if (!MEHD->RewriteChanges(RewriteChangeDelegate))
+		{
+			return SetWriteError();
+		}
+	}
+
+	// In standard mode the `moov` box has not been built yet. Do it and write it out.
+	if (Configuration.MuxMode == IMP4RawMuxer::FConfiguration::EMuxMode::Standard)
+	{
+		// Check to see if we really need to write composition offsets for the tracks.
+		for(auto &tcs : TrackChunkSamples)
+		{
+			if (!tcs.TrackMuxBoxes.CTTS)
+			{
+				continue;
+			}
+			// Is it needed?
+			if (!tcs.TrackMuxBoxes.CTTS->HaveCompositionOffsets())
+			{
+				// Not needed, remove it.
+				tcs.TrackMuxBoxes.CTTS->GetParent()->RemoveChild(tcs.TrackMuxBoxes.CTTS);
+				delete tcs.TrackMuxBoxes.CTTS;
+				tcs.TrackMuxBoxes.CTTS = nullptr;
+			}
+		}
+
+
+		MP4MuxerBoxes::FBoxBuilderData builderData;
+		MOOV->BuildData(builderData);
+		MOOV->CalculateBoxSizes();
+		bool bWriteOk = MOOV->VisitForWriting([&](MP4MuxerBoxes::FBoxBase* InBox)
+		{
+			InBox->SetAbsoluteFileOffset(FileHandle->Tell());
+			return InBox->GetCompiledBoxData().Num() ? FileHandle->Write(InBox->GetCompiledBoxData().GetData(), InBox->GetCompiledBoxData().Num()) : true;
+		});
+		if (!bWriteOk)
+		{
+			return SetWriteError();
+		}
+
+		// Get the final size of the `mdat` box.
+		const int64 SizeForMDATBox = FilePosAtEnd - FileOffsetOfStandardMDAT;
+		TArray<uint8> MDATBox;
+		int64 RewriteFileAt = FileOffsetOfStandardMDAT;
+		if (SizeForMDATBox > 0xffffffffU)
+		{
+			RewriteFileAt -= 8;
+			MDATBox.SetNumZeroed(16);
+			uint32* bd = (uint32*)MDATBox.GetData();
+			*bd++ = MP4MuxerBoxes::EndianSwap((uint32)1U);
+			*bd++ = MP4MuxerBoxes::EndianSwap(MP4MuxerBoxes::BoxAtom('m','d','a','t'));
+			uint64* bd2 = (uint64*)bd;
+			*bd2 = MP4MuxerBoxes::EndianSwap((uint64)(SizeForMDATBox + 8));
+		}
+		else
+		{
+			MDATBox.SetNumZeroed(8);
+			uint32* bd = (uint32*)MDATBox.GetData();
+			*bd++ = MP4MuxerBoxes::EndianSwap((uint32)SizeForMDATBox);
+			*bd++ = MP4MuxerBoxes::EndianSwap(MP4MuxerBoxes::BoxAtom('m','d','a','t'));
+		}
+		bWriteOk = FileHandle->Seek(RewriteFileAt) && FileHandle->Write(MDATBox.GetData(), MDATBox.Num()) && FileHandle->Seek(FilePosAtEnd);
+		if (!bWriteOk)
+		{
+			return SetWriteError();
+		}
+	}
+
+	// Need to change the major brand?
+	if (MinimumMajorISOBrand != WrittenMajorISOBrand)
+	{
+		if (MinimumMajorISOBrand <= 9)
+		{
+			if (!FTYP->UpdateMajorBrand(MP4MuxerBoxes::BoxAtom('i','s','o','0'+MinimumMajorISOBrand)))
+			{
+				return SetUpdateError(TEXT("major brand"), TEXT("ftyp"));
+			}
+		}
+		else if (MinimumMajorISOBrand <= 11)
+		{
+			if (!FTYP->UpdateMajorBrand(MP4MuxerBoxes::BoxAtom('i','s','o','a'+MinimumMajorISOBrand-10)))
+			{
+				return SetUpdateError(TEXT("major brand"), TEXT("ftyp"));
+			}
+		}
+		if (!FTYP->RewriteChanges(RewriteChangeDelegate))
+		{
+			return SetWriteError();
+		}
+	}
+
+	// Drop the boxes we needed for updating.
+	MOOV.Reset();
+	FTYP.Reset();
+	MVHD = nullptr;
+	MEHD = nullptr;
+	return true;
+}
+
+bool FMP4RawMuxer::RequestSampleData()
+{
+	TArray<FSampleRequest> TracksNeedingSamples;
+	FScopeLock lock(&Lock);
+	for(int32 i=0; i<TrackChunkSamples.Num(); ++i)
+	{
+		// Is this track still active?
+		if (!TrackChunkSamples[i].bHasFinished && !TrackChunkSamples[i].bGotLastSample)
+		{
+			// Does it need data?
+			if (TrackChunkSamples[i].AccumulatedDuration < TrackChunkSamples[i].InterleaveDurationNeeded)
+			{
+				FSampleRequest& sr(TracksNeedingSamples.Emplace_GetRef());
+				sr.TrackIndex = i;
+				sr.SampleNumber = TrackChunkSamples[i].NextSampleNumToRequest;
+			}
+		}
+	}
+	TracksAwaitingSamples = TracksNeedingSamples;
+	lock.Unlock();
+	SampleRequestDelegate.ExecuteIfBound(TracksNeedingSamples);
+	return TracksNeedingSamples.IsEmpty();
+}
+
+bool FMP4RawMuxer::AddTrackSample(int32 InTrack, const FTrackSample& InTrackSample)
+{
+	if (bFailed || LastError.Len())
+	{
+		return false;
+	}
+	FScopeLock lock(&Lock);
+	int32 srIdx = TracksAwaitingSamples.IndexOfByPredicate([&](const FSampleRequest& r){ return r.TrackIndex == InTrack; });
+	if (srIdx == INDEX_NONE)
+	{
+		return SetError(FString::Printf(TEXT("Called to provide sample data for track %d which has not requested sample data."), InTrack));
+	}
+	FTrackChunkSamples& tcs(TrackChunkSamples[InTrack]);
+	if (InTrackSample.SampleNumber != tcs.NextSampleNumToRequest && !InTrackSample.bIsFinalSample)
+	{
+		return SetError(FString::Printf(TEXT("Called to provide sample data for track %d sample number %u, but requested sample number %u."), InTrack, InTrackSample.SampleNumber, tcs.NextSampleNumToRequest));
+	}
+	TracksAwaitingSamples.RemoveAt(srIdx);
+	if (InTrackSample.Data.Num())
+	{
+		uint32 Duration = InTrackSample.Duration ? InTrackSample.Duration : TrackSpecs[InTrack].ConstantSampleDuration;
+		const bool bIsKeyframeOnly = TrackSpecs[InTrack].bIsAllKeyframes;
+		bool bIsKeyframe = bIsKeyframeOnly ? true : InTrackSample.bIsKeyframe;
+		if (Duration == 0)
+		{
+			return SetError(FString::Printf(TEXT("Sample duration must not be zero or less")));
+		}
+		// There is no version of a `stsz`, `stz2` or `trun` that would allow for a single AU to be larger than 4GiB (32 bits)
+		if (InTrackSample.Data.NumBytes() > 0xffffffffU)
+		{
+			return SetError(FString::Printf(TEXT("An ISO/IEC 14496-12 file can not store access units exceeding 4 GiB.")));
+		}
+		if (InTrackSample.SubSamples.Num() && !TrackSpecs[InTrack].SubSampleFlags.IsSet())
+		{
+			return SetError(FString::Printf(TEXT("Sample provides subsample information but the track was not configured for subsamples.")));
+		}
+
+		tcs.TotalTrackDuration += Duration;
+		tcs.AccumulatedDuration += Duration;
+		tcs.SumOfSampleDurations += Duration;
+
+		TUniquePtr<FTrackChunkSamples::FSample> ns = MakeUnique<FTrackChunkSamples::FSample>();
+		ns->SubSamples = InTrackSample.SubSamples;
+		ns->Data = InTrackSample.Data;
+		ns->DTS = InTrackSample.DTS;
+		ns->PTS = bIsKeyframeOnly ? InTrackSample.DTS : InTrackSample.PTS;
+		ns->Duration = Duration;
+		ns->bIsKeyframe = bIsKeyframe;
+		tcs.Samples.Emplace(MoveTemp(ns));
+		++tcs.NextSampleNumToRequest;
+	}
+	tcs.bGotLastSample = InTrackSample.bIsFinalSample;
+	lock.Unlock();
+	if (WorkerEvent)
+	{
+		WorkerEvent->Trigger();
+	}
+	return true;
+}
+
+
+bool FMP4RawMuxer::WriteAccumulatedSamples()
+{
+	int32 NumActiveTracks = 0;
+	for(int32 i=0; i<TrackChunkSamples.Num(); ++i)
+	{
+		NumActiveTracks += !TrackChunkSamples[i].bHasFinished ? 1 : 0;
+	}
+	if (NumActiveTracks == 0)
+	{
+		return true;
+	}
+	// Write standard or fragmented.
+	bool bOk = Configuration.MuxMode == IMP4RawMuxer::FConfiguration::EMuxMode::Standard ? WriteAccumulatedSamplesStandard() : WriteAccumulatedSamplesFragmented();
+	// Dump the samples that were written.
+	for(auto &tcs : TrackChunkSamples)
+	{
+		tcs.Empty();
+		if (tcs.bGotLastSample)
+		{
+			tcs.bHasFinished = true;
+		}
+	}
+	// Continue with next batch of samples.
+	WorkerEvent->Trigger();
+	return bOk;
+}
+
+bool FMP4RawMuxer::WriteAccumulatedSamplesStandard()
+{
+	for(int32 nTrk=0; nTrk<TrackChunkSamples.Num(); ++nTrk)
+	{
+		const FTrackChunkSamples& tcs(TrackChunkSamples[nTrk]);
+		// Skip if track no longer active.
+		if (tcs.bHasFinished || tcs.Samples.IsEmpty())
+		{
+			continue;
+		}
+		if (!tcs.TrackMuxBoxes.STSZ || !tcs.TrackMuxBoxes.STTS || !tcs.TrackMuxBoxes.CTTS || !tcs.TrackMuxBoxes.STSC || !tcs.TrackMuxBoxes.STCO)
+		{
+			return SetError(TEXT("Internal error, one of the required boxes has not been created yet."));
+		}
+
+		// Where in the file are we at now?
+		int64 FileDataOffset = FileHandle->Tell();
+		// Write the samples of this chunk.
+		for(int32 i=0; i<tcs.Samples.Num(); ++i)
+		{
+			const FTrackChunkSamples::FSample* const Smp = tcs.Samples[i].Get();
+			bool bWriteOk = FileHandle->Write(Smp->Data.GetData(), Smp->Data.Num());
+			if (!bWriteOk)
+			{
+				return SetWriteError();
+			}
+			// Set the sample size. Do this first since that allows us to see how many samples there are in total.
+			tcs.TrackMuxBoxes.STSZ->AddSampleSize((uint32) Smp->Data.Num());
+			// Set the sample duration.
+			tcs.TrackMuxBoxes.STTS->AddSampleDuration(Smp->Duration);
+			// Sync sample is optional to write.
+			if (Smp->bIsKeyframe && tcs.TrackMuxBoxes.STSS)
+			{
+				tcs.TrackMuxBoxes.STSS->AddSampleIndex(tcs.TrackMuxBoxes.STSZ->GetNumberOfSamples());
+			}
+
+			// On the very first sample, if it has a negative DTS we might want to add an edit list.
+			if (Smp->DTS < 0 && tcs.TrackMuxBoxes.STSZ->GetNumberOfSamples() == 1)
+			{
+				MP4MuxerBoxes::FBoxEDTS* EDTS = new MP4MuxerBoxes::FBoxEDTS;
+				tcs.TrackMuxBoxes.TKHD->GetParent()->AddChildAfter(EDTS, tcs.TrackMuxBoxes.TKHD);
+				MP4MuxerBoxes::FBoxELST* ELST = new MP4MuxerBoxes::FBoxELST;
+				TrackChunkSamples[nTrk].TrackMuxBoxes.ELST = ELST;
+				EDTS->AddChild(ELST);
+				ELST->SetMediaTime(-Smp->DTS);
+			}
+
+			// Set composition time offset.
+			int32 CompositionTimeOffset = Smp->PTS - Smp->DTS;
+			tcs.TrackMuxBoxes.CTTS->AddSampleCompositionOffset(CompositionTimeOffset);
+
+			// Subsamples?
+			if (Smp->SubSamples.Num() && tcs.TrackMuxBoxes.SUBS)
+			{
+				tcs.TrackMuxBoxes.SUBS->AddSubsamplesForSample(tcs.TrackMuxBoxes.STSZ->GetNumberOfSamples(), Smp->SubSamples);
+			}
+		}
+		// Negative composition offsets require the major brand to be at least 4.
+		if (MinimumMajorISOBrand < 4 && tcs.TrackMuxBoxes.CTTS->HasNegativeCompositionOffsets() && tcs.TrackMuxBoxes.CTTS->HaveCompositionOffsets())
+		{
+			MinimumMajorISOBrand = 4;
+		}
+		// Add chunk offset.
+		tcs.TrackMuxBoxes.STCO->AddChunkOffset(FileDataOffset);
+		// Add sample to chunk entry.
+		tcs.TrackMuxBoxes.STSC->AddEntry(tcs.TrackMuxBoxes.STCO->GetNumberOfChunkOffsets(), (uint32)tcs.Samples.Num(), 1U);
+	}
+	return true;
+}
+
+bool FMP4RawMuxer::WriteAccumulatedSamplesFragmented()
+{
+	// On the very first sample, if it has a negative DTS we might want to add an edit list.
+	if (FragmentSequenceNumber == 0)
+	{
+		for(int32 nTrk=0; nTrk<TrackChunkSamples.Num(); ++nTrk)
+		{
+			const FTrackChunkSamples& tcs(TrackChunkSamples[nTrk]);
+			if (tcs.Samples.Num() && tcs.Samples[0]->DTS < 0 && !tcs.TrackMuxBoxes.ELST)
+			{
+				MP4MuxerBoxes::FBoxEDTS* EDTS = new MP4MuxerBoxes::FBoxEDTS;
+				tcs.TrackMuxBoxes.TKHD->GetParent()->AddChildAfter(EDTS, tcs.TrackMuxBoxes.TKHD);
+
+				MP4MuxerBoxes::FBoxELST* ELST = new MP4MuxerBoxes::FBoxELST;
+				TrackChunkSamples[nTrk].TrackMuxBoxes.ELST = ELST;
+				EDTS->AddChild(ELST);
+				ELST->SetMediaTime(-tcs.Samples[0]->DTS);
+			}
+		}
+		// Write out the `moov` box now.
+		if (!bWroteMOOV)
+		{
+			bWroteMOOV = true;
+			MP4MuxerBoxes::FBoxBuilderData builderData;
+			MOOV->BuildData(builderData);
+			MOOV->CalculateBoxSizes();
+			bool bWriteOk = MOOV->VisitForWriting([&](MP4MuxerBoxes::FBoxBase* InBox)
+			{
+				InBox->SetAbsoluteFileOffset(FileHandle->Tell());
+				return InBox->GetCompiledBoxData().Num() ? FileHandle->Write(InBox->GetCompiledBoxData().GetData(), InBox->GetCompiledBoxData().Num()) : true;
+			});
+			if (!bWriteOk)
+			{
+				return SetWriteError();
+			}
+		}
+	}
+
+	TUniquePtr<MP4MuxerBoxes::FBoxMOOF> MOOF;
+	struct FOffsetPatch
+	{
+		MP4MuxerBoxes::FBoxTFHD* TFHDBox = nullptr;
+		int64 FileOffset = -1;
+	};
+	TArray<FOffsetPatch> TFHDBoxesToPatch;
+	int64 SizeOfAllSamples = 0;
+	for(int32 nTrk=0; nTrk<TrackChunkSamples.Num(); ++nTrk)
+	{
+		const FTrackChunkSamples& tcs(TrackChunkSamples[nTrk]);
+		// Skip if track no longer active or has no samples
+		if (tcs.bHasFinished || tcs.Samples.IsEmpty())
+		{
+			continue;
+		}
+		if (!MOOF)
+		{
+			MOOF = MakeUnique<MP4MuxerBoxes::FBoxMOOF>();
+			MP4MuxerBoxes::FBoxMFHD* MFHD = new MP4MuxerBoxes::FBoxMFHD;
+			MOOF->AddChild(MFHD);
+			MFHD->SetSequenceNumber(++FragmentSequenceNumber);
+		}
+
+		MP4MuxerBoxes::FBoxTRAF* TRAF = new MP4MuxerBoxes::FBoxTRAF;
+		MOOF->AddChild(TRAF);
+		MP4MuxerBoxes::FBoxTFHD* TFHD = new MP4MuxerBoxes::FBoxTFHD;
+		TRAF->AddChild(TFHD);
+		MP4MuxerBoxes::FBoxTFDT* TFDT = new MP4MuxerBoxes::FBoxTFDT;
+		TRAF->AddChild(TFDT);
+		TFDT->SetBaseMediaDecodeTime((uint64)(tcs.TotalTrackDuration - tcs.SumOfSampleDurations));
+		TFHD->SetTrackID(1U + nTrk);
+		TFHD->SetBaseDataOffset(0);
+		/*
+			If using default-base-is-moof the major brand needs to be at least 5.
+
+			if (MinimumMajorISOBrand < 5)
+			{
+				MinimumMajorISOBrand = 5;
+			}
+		*/
+		TFHDBoxesToPatch.Add({TFHD});
+
+		const bool bAllAreSyncSamples = TrackSpecs[nTrk].bIsAllKeyframes;
+		const uint32 DefaultSampleFlags = tcs.TrackMuxBoxes.TREX ? tcs.TrackMuxBoxes.TREX->GetDefaultSampleFlags() : 0U;
+
+		MP4MuxerBoxes::FBoxTRUN* TRUN = new MP4MuxerBoxes::FBoxTRUN;
+		TRAF->AddChild(TRUN);
+		TRUN->SetDefaultSampleFlagsForTesting(DefaultSampleFlags);
+
+		// Sub-sample information
+		MP4MuxerBoxes::FBoxSUBS* SUBS = nullptr;
+
+		// Process samples.
+		for(int32 i=0; i<tcs.Samples.Num(); ++i)
+		{
+			const FTrackChunkSamples::FSample* const Smp = tcs.Samples[i].Get();
+			uint32 Flags;
+			if (bAllAreSyncSamples || Smp->bIsKeyframe)
+			{
+				Flags = DefaultSampleFlags & ~0x10000U;
+			}
+			else
+			{
+				Flags = DefaultSampleFlags | 0x10000U;
+			}
+
+			int32 CompositionTimeOffset = Smp->PTS - Smp->DTS;
+			// Negative composition offsets in the `trun` require the major brand to be at least 6.
+			if (CompositionTimeOffset < 0 && MinimumMajorISOBrand < 6)
+			{
+				MinimumMajorISOBrand = 6;
+			}
+			TRUN->AddSample(Smp->Duration, Smp->Data.Num(), Flags, CompositionTimeOffset);
+			// Accumulate the size of all samples, which is how big the following `mdat` box
+			// is going to be. This could exceed 4 GiB which may require a special size indicator
+			// which is difficult to patch in later, so it's best if we know upfront if we need it.
+			SizeOfAllSamples += Smp->Data.Num();
+
+			// Subsamples?
+			if (Smp->SubSamples.Num() && TrackSpecs[nTrk].SubSampleFlags.IsSet())
+			{
+				if (!SUBS)
+				{
+					SUBS = new MP4MuxerBoxes::FBoxSUBS;
+					TRAF->AddChild(SUBS);
+					SUBS->SetFlags(TrackSpecs[nTrk].SubSampleFlags.GetValue());
+				}
+				SUBS->AddSubsamplesForSample(i + 1, Smp->SubSamples);
+			}
+		}
+		// If there is no constant sample duration given for the track in general but the
+		// samples in this run all have the same duration, set the constant duration for this run.
+		if (!TRUN->HasDifferentDurations())
+		{
+			// If there is a constant sample duration for the track and it's the same
+			// as the constant duration of the samples in this run then we do not need
+			// to write per-sample durations.
+			if (TrackSpecs[nTrk].ConstantSampleDuration == tcs.Samples[0]->Duration)
+			{
+				TRUN->DoNotWriteDurations();
+			}
+			// If there is no constant sample duration set for the track then put this
+			// run's constant duration into the track fragment header.
+			else if (TrackSpecs[nTrk].ConstantSampleDuration == 0)
+			{
+				TFHD->SetDefaultSampleDuration(tcs.Samples[0]->Duration);
+				TRUN->DoNotWriteDurations();
+			}
+		}
+
+		// Same for the sizes.
+		// TODO: check that if fixed size is given it matches the sample size!
+		if (/*TrackSpecs[nTrk].ConstantSampleSize == 0 &&*/ !TRUN->HasDifferentSizes())
+		{
+			TFHD->SetDefaultSampleSize((uint32) tcs.Samples[0]->Data.Num());
+			TRUN->DoNotWriteSizes();
+		}
+		// When sample flags differ during the run then the flags need to be written and we can skip
+		// any special handling.
+		if (!TRUN->AreFlagsAfterFirstSampleDifferent())
+		{
+			// If only the first sample has flags other than the default while the rest uses default, set first sample flags only.
+			if (TRUN->GetSubsequentSampleFlagsOrDefault() == DefaultSampleFlags && TRUN->GetFirstSampleFlagsOrDefault() != DefaultSampleFlags)
+			{
+				TRUN->SetFirstSampleFlags(TRUN->GetFirstSampleFlagsOrDefault());
+			}
+			// Are all samples using default value?
+			else if (TRUN->GetFirstSampleFlagsOrDefault() == TRUN->GetSubsequentSampleFlagsOrDefault() && TRUN->GetFirstSampleFlagsOrDefault() == DefaultSampleFlags)
+			{
+				TRUN->DoNotWriteFlags();
+			}
+			else
+			{
+				// The samples after the first all have the same flags. Set that in the trak header box.
+				TFHD->SetDefaultSampleFlags(TRUN->GetSubsequentSampleFlagsOrDefault());
+				// If the first sample has a different value set it up.
+				if (TRUN->GetFirstSampleFlagsOrDefault() != TRUN->GetSubsequentSampleFlagsOrDefault())
+				{
+					TRUN->SetFirstSampleFlags(TRUN->GetFirstSampleFlagsOrDefault());
+				}
+				TRUN->DoNotWriteFlags();
+			}
+		}
+	}
+	if (!MOOF)
+	{
+		return true;
+	}
+	MP4MuxerBoxes::FBoxBuilderData builderData;
+	MOOF->BuildData(builderData);
+	MOOF->CalculateBoxSizes();
+	bool bWriteOk = MOOF->VisitForWriting([&](MP4MuxerBoxes::FBoxBase* InBox)
+	{
+		InBox->SetAbsoluteFileOffset(FileHandle->Tell());
+		return InBox->GetCompiledBoxData().Num() ? FileHandle->Write(InBox->GetCompiledBoxData().GetData(), InBox->GetCompiledBoxData().Num()) : true;
+	});
+	if (!bWriteOk)
+	{
+		return SetWriteError();
+	}
+
+	// Write the `mdat` box.
+	int64 SizeForMDATBox = 8 + SizeOfAllSamples;
+	TArray<uint8> MDATBox;
+	if (SizeForMDATBox > 0xffffffffU)
+	{
+		MDATBox.SetNumZeroed(16);
+		uint32* bd = (uint32*)MDATBox.GetData();
+		*bd++ = MP4MuxerBoxes::EndianSwap((uint32)1U);
+		*bd++ = MP4MuxerBoxes::EndianSwap(MP4MuxerBoxes::BoxAtom('m','d','a','t'));
+		uint64* bd2 = (uint64*)bd;
+		*bd2 = MP4MuxerBoxes::EndianSwap((uint64)(SizeForMDATBox + 8));
+	}
+	else
+	{
+		MDATBox.SetNumZeroed(8);
+		uint32* bd = (uint32*)MDATBox.GetData();
+		*bd++ = MP4MuxerBoxes::EndianSwap((uint32)SizeForMDATBox);
+		*bd++ = MP4MuxerBoxes::EndianSwap(MP4MuxerBoxes::BoxAtom('m','d','a','t'));
+	}
+	bWriteOk = FileHandle->Write(MDATBox.GetData(), MDATBox.Num());
+	if (!bWriteOk)
+	{
+		return SetWriteError();
+	}
+
+	// Write the samples
+	int64 FileDataOffset = FileHandle->Tell();
+	int32 PatchNum = 0;
+	for(int32 nTrk=0; nTrk<TrackChunkSamples.Num(); ++nTrk)
+	{
+		const FTrackChunkSamples& tcs(TrackChunkSamples[nTrk]);
+		// Skip if track no longer active or has no samples
+		if (tcs.bHasFinished || tcs.Samples.IsEmpty())
+		{
+			continue;
+		}
+		TFHDBoxesToPatch[PatchNum++].FileOffset = FileHandle->Tell();
+		for(int32 i=0; i<tcs.Samples.Num(); ++i)
+		{
+			const FTrackChunkSamples::FSample* const Smp = tcs.Samples[i].Get();
+			bWriteOk = FileHandle->Write(Smp->Data.GetData(), Smp->Data.Num());
+			if (!bWriteOk)
+			{
+				return SetWriteError();
+			}
+		}
+	}
+	// Patch the offset in the `tfhd` boxes
+	for(int32 i=0; i<TFHDBoxesToPatch.Num(); ++i)
+	{
+		if (!TFHDBoxesToPatch[i].TFHDBox->UpdateBaseDataOffset((uint64)TFHDBoxesToPatch[i].FileOffset))
+		{
+			return SetUpdateError(TEXT("base data offset"), TEXT("tfhd"));
+		}
+		bWriteOk = TFHDBoxesToPatch[i].TFHDBox->RewriteChanges(RewriteChangeDelegate);
+		if (!bWriteOk)
+		{
+			return SetWriteError();
+		}
+	}
+	return true;
+}
+
+
+uint32 FMP4RawMuxer::Run()
+{
+	// Open the file
+	FileHandle = FPlatformFileManager::Get().GetPlatformFile().OpenWrite(*Configuration.OutputFilename, false, false);
+	if (!FileHandle)
+	{
+		SetError(FString::Printf(TEXT("Could not open file \"%s\" for writing."), *Configuration.OutputFilename));
+		StatusDelegate.ExecuteIfBound(EStatus::Failed);
+	}
+	else if (!WriteFileStart())
+	{
+		SetWriteError();
+		StatusDelegate.ExecuteIfBound(EStatus::Failed);
+	}
+
+	const FTimespan WaitTime(FTimespan::FromMilliseconds(100));
+	WorkerEvent->Trigger();
+	bool bIsDone = false;
+	while(!bTerminate)
+	{
+		bool bGot = WorkerEvent->Wait(WaitTime);
+		if (!bGot)
+		{
+			continue;
+		}
+		WorkerEvent->Reset();
+		if (bFailed || bIsDone)
+		{
+			continue;
+		}
+		// Continue requesting samples until we have enough data for interleaving.
+		bool bGotAllSamples = RequestSampleData();
+		if (bGotAllSamples)
+		{
+			// Write them
+			if (!WriteAccumulatedSamples())
+			{
+				SetWriteError();
+				StatusDelegate.ExecuteIfBound(EStatus::Failed);
+				continue;
+			}
+			// Are we done?
+			bool bAllDone = true;
+			for(auto &tcs : TrackChunkSamples)
+			{
+				if (!tcs.bHasFinished)
+				{
+					bAllDone = false;
+					break;
+				}
+			}
+			if (bAllDone)
+			{
+				bool bOk = WriteFileEnd();
+				CloseFile();
+				if (bOk)
+				{
+					StatusDelegate.ExecuteIfBound(EStatus::Finished);
+				}
+				else
+				{
+					StatusDelegate.ExecuteIfBound(EStatus::Failed);
+				}
+				bIsDone = true;
+			}
+		}
+	}
+	CloseFile();
+	return 0;
+}
diff --git a/Engine/Plugins/Media/MP4Utilities/Source/MP4Muxer/Private/MP4MuxerBoxes.cpp b/Engine/Plugins/Media/MP4Utilities/Source/MP4Muxer/Private/MP4MuxerBoxes.cpp
new file mode 100644
--- /dev/null
+++ b/Engine/Plugins/Media/MP4Utilities/Source/MP4Muxer/Private/MP4MuxerBoxes.cpp
@@ -0,0 +1,1219 @@
+// Copyright Epic Games, Inc. All Rights Reserved.
+
+#include "MP4MuxerBoxes.h"
+#include "Math/NumericLimits.h"
+
+namespace MP4MuxerBoxes
+{
+
+namespace
+{
+
+class FBoxDataWriter
+{
+public:
+	FBoxDataWriter(TArray<uint8>& InOutBuffer)
+		: Buffer(InOutBuffer)
+	{
+		BytePosition = Buffer.Num();
+	}
+	~FBoxDataWriter()
+	{
+		Close();
+	}
+	void Close()
+	{
+		AlignToBytes(0);
+		Buffer.SetNum(BytePosition, EAllowShrinking::No);
+	}
+
+	uint32 PutAlignedValue(uint8 InValue)
+	{
+		check(IsByteAligned());
+		PutByte(InValue);
+		return BytePosition - 1;
+	}
+	uint32 PutAlignedValue(uint16 InValue)
+	{
+		PutAlignedValue((uint8)(InValue >> 8));
+		PutAlignedValue((uint8)(InValue & 0xff));
+		return BytePosition - 2;
+	}
+	uint32 PutAlignedValue(uint32 InValue)
+	{
+		PutAlignedValue((uint16)(InValue >> 16));
+		PutAlignedValue((uint16)(InValue & 0xffff));
+		return BytePosition - 4;
+	}
+	uint32 PutAlignedValue(uint64 InValue)
+	{
+		PutAlignedValue((uint32)(InValue >> 32));
+		PutAlignedValue((uint32)(InValue & 0xffffffff));
+		return BytePosition - 8;
+	}
+	uint32 PutAlignedValue(int8 InValue)
+	{
+		return PutAlignedValue((uint8)InValue);
+	}
+	uint32 PutAlignedValue(int16 InValue)
+	{
+		return PutAlignedValue((uint16)InValue);
+	}
+	uint32 PutAlignedValue(int32 InValue)
+	{
+		return PutAlignedValue((uint32)InValue);
+	}
+	uint32 PutAlignedValue(int64 InValue)
+	{
+		return PutAlignedValue((uint64)InValue);
+	}
+	uint32 PutAlignedArray(TConstArrayView<const uint8> InArray)
+	{
+		for(int32 i=0, iMax=InArray.Num(); i<iMax; ++i)
+		{
+			PutAlignedValue(InArray[i]);
+		}
+		return BytePosition - InArray.Num();
+	}
+
+	void SetAlignedValueAt(uint8 InValue, uint32 InAtOffset)
+	{
+		check(InAtOffset < BytePosition);
+		Buffer[InAtOffset] = InValue;
+	}
+	void SetAlignedValueAt(uint16 InValue, uint32 InAtOffset)
+	{
+		SetAlignedValueAt((uint8)(InValue >> 8), InAtOffset);
+		SetAlignedValueAt((uint8)(InValue & 0xff), InAtOffset + 1);
+	}
+	void SetAlignedValueAt(uint32 InValue, uint32 InAtOffset)
+	{
+		SetAlignedValueAt((uint16)(InValue >> 16), InAtOffset);
+		SetAlignedValueAt((uint16)(InValue & 0xffff), InAtOffset + 2);
+	}
+	void SetAlignedValueAt(uint64 InValue, uint32 InAtOffset)
+	{
+		SetAlignedValueAt((uint32)(InValue >> 32), InAtOffset);
+		SetAlignedValueAt((uint32)(InValue & 0xffffffff), InAtOffset + 4);
+	}
+
+	void SetAlignedValueAt(int8 InValue, uint32 InAtOffset)
+	{
+		SetAlignedValueAt((uint8)InValue, InAtOffset);
+	}
+	void SetAlignedValueAt(int16 InValue, uint32 InAtOffset)
+	{
+		SetAlignedValueAt((uint16)InValue, InAtOffset);
+	}
+	void SetAlignedValueAt(int32 InValue, uint32 InAtOffset)
+	{
+		SetAlignedValueAt((uint32)InValue, InAtOffset);
+	}
+	void SetAlignedValueAt(int64 InValue, uint32 InAtOffset)
+	{
+		SetAlignedValueAt((uint64)InValue, InAtOffset);
+	}
+
+
+	void PutBits(uint32 InValue, uint32 InNumBits)
+	{
+		check(InNumBits != 0);
+		check(InNumBits <= 32);
+		uint32 Mask = 1U << (InNumBits-1);
+		while(Mask)
+		{
+			PutBit((InValue & Mask) != 0 ? 1 : 0);
+			Mask >>= 1;
+		}
+	}
+	void PutBits64(uint64 InValue, uint32 InNumBits)
+	{
+		uint32 Upper = static_cast<uint32>(InValue >> 32U);
+		uint32 Lower = static_cast<uint32>(InValue);
+		if (InNumBits > 32)
+		{
+			const uint32 nb = InNumBits - 32;
+			PutBits(Upper, nb);
+			InNumBits -= nb;
+		}
+		PutBits(Lower, InNumBits);
+	}
+	void AlignToBytes(uint32 InFillBitsWith)
+	{
+		if (BitPosition)
+		{
+			do
+			{
+				PutBit(InFillBitsWith);
+			} while(BitPosition);
+		}
+	}
+	uint32 GetNumBits() const
+	{
+		return BytePosition * 8 + BitPosition;
+	}
+	uint32 GetNumBytesUsed() const
+	{
+		return (GetNumBits() + 7) / 8;
+	}
+	bool IsByteAligned() const
+	{
+		return BitPosition == 0;
+	}
+	uint32 GetBytePosition() const
+	{
+		return BytePosition;
+	}
+private:
+	void PutByte(uint8 InValue)
+	{
+		if (BytePosition == (uint32)Buffer.Num())
+		{
+			Buffer.SetNumZeroed(Buffer.Num() + 256);
+		}
+		Buffer[BytePosition++] = InValue;
+	}
+	void PutBit(uint32 InBit)
+	{
+		InBit <<= (7 - BitPosition);
+		check(Buffer.Num());
+		Buffer[BytePosition] |= (uint8)InBit;
+		if (++BitPosition == 8)
+		{
+			BitPosition = 0;
+			if (++BytePosition > (uint32)Buffer.Num())
+			{
+				Buffer.SetNumZeroed(Buffer.Num() + 256);
+			}
+		}
+	}
+	TArray<uint8>& Buffer;
+	uint32 BytePosition = 0;
+	uint32 BitPosition = 0;
+};
+
+
+}
+
+
+
+
+
+
+
+FBoxBase::~FBoxBase()
+{
+	for(int32 i=0; i<Children.Num(); ++i)
+	{
+		delete Children[i];
+	}
+}
+
+
+bool FBoxBase::BuildData(const FBoxBuilderData& InOutBuilderData)
+{
+	// Build children recursively first.
+	for(int32 i=0; i<Children.Num(); ++i)
+	{
+		Children[i]->BuildData(InOutBuilderData);
+	}
+	// If this is an artificial 'root' box, that is, the file level, then we are done.
+	if (Type == 0)
+	{
+		return true;
+	}
+
+	FBoxDataWriter dw(CompiledBoxData);
+	CompiledBoxData_Offset_Size = dw.PutAlignedValue(Size);
+	dw.PutAlignedValue(Type);
+	Size = dw.GetBytePosition();
+	return true;
+}
+
+void FBoxBase::CalculateBoxSizes()
+{
+	for(int32 i=0; i<Children.Num(); ++i)
+	{
+		Children[i]->CalculateBoxSizes();
+		Size += Children[i]->Size;
+	}
+	if (Type)
+	{
+		FBoxDataWriter dw(CompiledBoxData);
+		dw.SetAlignedValueAt(Size, CompiledBoxData_Offset_Size);
+	}
+}
+
+
+
+bool FBoxFull::BuildData(const FBoxBuilderData& InOutBuilderData)
+{
+	FBoxBase::BuildData(InOutBuilderData);
+	FBoxDataWriter dw(CompiledBoxData);
+	dw.PutAlignedValue(Version);
+	dw.PutBits(Flags, 24);
+	check(dw.IsByteAligned());
+	Size = dw.GetBytePosition();
+	return true;
+}
+
+
+bool FBoxFTYP::BuildData(const FBoxBuilderData& InOutBuilderData)
+{
+	FBoxBase::BuildData(InOutBuilderData);
+	FBoxDataWriter dw(CompiledBoxData);
+	MajorBrandPatchOffset = dw.PutAlignedValue(MajorBrand);
+	dw.PutAlignedValue(MinorVersion);
+	for(int32 i=0; i<CompatibleBrands.Num(); ++i)
+	{
+		dw.PutAlignedValue(CompatibleBrands[i]);
+	}
+	Size = dw.GetBytePosition();
+	bWasBuilt = true;
+	return true;
+}
+
+bool FBoxFTYP::UpdateMajorBrand(uint32 InMajorBrand)
+{
+	SetMajorBrand(InMajorBrand);
+	if (!bWasBuilt)
+	{
+		return true;
+	}
+	FBoxDataWriter dw(CompiledBoxData);
+	dw.SetAlignedValueAt(MajorBrand, MajorBrandPatchOffset);
+	return true;
+}
+bool FBoxFTYP::RewriteChanges(FRewriteBoxChangeDelegate InRewriteBoxChangeDelegate)
+{
+	// If this box was not yet built then there's noting to rewrite.
+	if (!bWasBuilt)
+	{
+		return true;
+	}
+	return InRewriteBoxChangeDelegate.Execute(AbsoluteFileOffset + MajorBrandPatchOffset, MakeConstArrayView(CompiledBoxData.GetData() + MajorBrandPatchOffset, 4));
+}
+
+
+bool FBoxFREE::BuildData(const FBoxBuilderData& InOutBuilderData)
+{
+	FBoxBase::BuildData(InOutBuilderData);
+	FBoxDataWriter dw(CompiledBoxData);
+	for(uint32 i=0; i<NumEmptyBytes; ++i)
+	{
+		dw.PutAlignedValue((uint8)0);
+	}
+	Size = dw.GetBytePosition();
+	return true;
+}
+
+bool FBoxMVHD::BuildData(const FBoxBuilderData& InOutBuilderData)
+{
+	// Any of the values require a version 1 box?
+	Version = (CreationTime > 0xffffffffU || ModificationTime > 0xffffffffU || (Duration != ~((uint64)0) && Duration > 0xffffffffU)) ? 1 : 0;
+
+	FBoxFull::BuildData(InOutBuilderData);
+	FBoxDataWriter dw(CompiledBoxData);
+	if (Version == 1)
+	{
+		dw.PutAlignedValue(CreationTime);
+		dw.PutAlignedValue(ModificationTime);
+		TimescalePatchOffset = dw.PutAlignedValue(Timescale);
+		DurationPatchOffset = dw.PutAlignedValue(Duration);
+	}
+	else
+	{
+		dw.PutAlignedValue((uint32)CreationTime);
+		dw.PutAlignedValue((uint32)ModificationTime);
+		TimescalePatchOffset = dw.PutAlignedValue(Timescale);
+		DurationPatchOffset = dw.PutAlignedValue((uint32)Duration);
+	}
+	dw.PutAlignedValue((uint32)0x00010000);	// rate of 1.0
+	dw.PutAlignedValue((uint16)0x0100);		// volume 1.0
+	dw.PutAlignedValue((uint16)0);			// reserved
+	dw.PutAlignedValue((uint32)0);			// reserved
+	dw.PutAlignedValue((uint32)0);			// reserved
+	// unity matrix
+	dw.PutAlignedValue((uint32)0x10000);
+	dw.PutAlignedValue((uint32)0);
+	dw.PutAlignedValue((uint32)0);
+	dw.PutAlignedValue((uint32)0);
+	dw.PutAlignedValue((uint32)0x10000);
+	dw.PutAlignedValue((uint32)0);
+	dw.PutAlignedValue((uint32)0);
+	dw.PutAlignedValue((uint32)0x10000);
+	dw.PutAlignedValue((uint32)0x40000000);
+	// predefined
+	dw.PutAlignedValue((uint32)0);
+	dw.PutAlignedValue((uint32)0);
+	dw.PutAlignedValue((uint32)0);
+	dw.PutAlignedValue((uint32)0);
+	dw.PutAlignedValue((uint32)0);
+	dw.PutAlignedValue((uint32)0);
+	// next_track_ID
+	dw.PutAlignedValue(NextTrackID);
+	Size = dw.GetBytePosition();
+	bWasBuilt = true;
+	return true;
+}
+bool FBoxMVHD::UpdateTimescale(uint32 InTimescale)
+{
+	if (!bWasBuilt)
+	{
+		SetTimescale(InTimescale);
+		return true;
+	}
+	SetTimescale(InTimescale);
+	FBoxDataWriter dw(CompiledBoxData);
+	dw.SetAlignedValueAt(Timescale, TimescalePatchOffset);
+	return true;
+}
+bool FBoxMVHD::UpdateDuration(uint64 InDuration)
+{
+	if (!bWasBuilt)
+	{
+		SetDuration(InDuration);
+		return true;
+	}
+	if (InDuration > 0xffffffffU && Version == 0)
+	{
+		return false;
+	}
+	SetDuration(InDuration);
+	FBoxDataWriter dw(CompiledBoxData);
+	if (Version == 1)
+	{
+		dw.SetAlignedValueAt(Duration, DurationPatchOffset);
+	}
+	else
+	{
+		dw.SetAlignedValueAt((uint32)Duration, DurationPatchOffset);
+	}
+	return true;
+}
+bool FBoxMVHD::RewriteChanges(FRewriteBoxChangeDelegate InRewriteBoxChangeDelegate)
+{
+	// If this box was not yet built then there's noting to rewrite.
+	if (!bWasBuilt)
+	{
+		return true;
+	}
+	return InRewriteBoxChangeDelegate.Execute(AbsoluteFileOffset + TimescalePatchOffset, MakeConstArrayView(CompiledBoxData.GetData() + TimescalePatchOffset, 4)) &&
+		   InRewriteBoxChangeDelegate.Execute(AbsoluteFileOffset + DurationPatchOffset, MakeConstArrayView(CompiledBoxData.GetData() + DurationPatchOffset, Version == 1 ? 8 : 4));
+}
+
+
+bool FBoxTKHD::BuildData(const FBoxBuilderData& InOutBuilderData)
+{
+	// Any of the values require a version 1 box?
+	Version = (CreationTime > 0xffffffffU || ModificationTime > 0xffffffffU || (Duration != ~((uint64)0) && Duration > 0xffffffffU)) ? 1 : 0;
+
+	FBoxFull::BuildData(InOutBuilderData);
+	FBoxDataWriter dw(CompiledBoxData);
+	if (Version == 1)
+	{
+		dw.PutAlignedValue(CreationTime);
+		dw.PutAlignedValue(ModificationTime);
+		dw.PutAlignedValue(TrackID);
+		dw.PutAlignedValue((uint32)0);
+		DurationPatchOffset = dw.PutAlignedValue(Duration);
+	}
+	else
+	{
+		dw.PutAlignedValue((uint32)CreationTime);
+		dw.PutAlignedValue((uint32)ModificationTime);
+		dw.PutAlignedValue(TrackID);
+		dw.PutAlignedValue((uint32)0);
+		DurationPatchOffset = dw.PutAlignedValue((uint32)Duration);
+	}
+	// reserved
+	dw.PutAlignedValue((uint32)0);
+	dw.PutAlignedValue((uint32)0);
+	dw.PutAlignedValue(Layer);
+	dw.PutAlignedValue(AlternateGroup);
+	dw.PutAlignedValue(Volume);
+	dw.PutAlignedValue((uint16)0);
+	// unity matrix
+	dw.PutAlignedValue((uint32)0x10000);
+	dw.PutAlignedValue((uint32)0);
+	dw.PutAlignedValue((uint32)0);
+	dw.PutAlignedValue((uint32)0);
+	dw.PutAlignedValue((uint32)0x10000);
+	dw.PutAlignedValue((uint32)0);
+	dw.PutAlignedValue((uint32)0);
+	dw.PutAlignedValue((uint32)0x10000);
+	dw.PutAlignedValue((uint32)0x40000000);
+	dw.PutAlignedValue(Width);
+	dw.PutAlignedValue(Height);
+	Size = dw.GetBytePosition();
+	bWasBuilt = true;
+	return true;
+}
+bool FBoxTKHD::UpdateDuration(uint64 InDuration)
+{
+	if (!bWasBuilt)
+	{
+		SetDuration(InDuration);
+		return true;
+	}
+	// If the duration now requires 64 bit wide fields when originally we had only 32 then this cannot be updated.
+	if (InDuration > 0xffffffffU && Version == 0)
+	{
+		return false;
+	}
+	SetDuration(InDuration);
+	FBoxDataWriter dw(CompiledBoxData);
+	if (Version == 1)
+	{
+		dw.SetAlignedValueAt(Duration, DurationPatchOffset);
+	}
+	else
+	{
+		dw.SetAlignedValueAt((uint32)Duration, DurationPatchOffset);
+	}
+	return true;
+}
+bool FBoxTKHD::RewriteChanges(FRewriteBoxChangeDelegate InRewriteBoxChangeDelegate)
+{
+	// If this box was not yet built then there's noting to rewrite.
+	if (!bWasBuilt)
+	{
+		return true;
+	}
+	return InRewriteBoxChangeDelegate.Execute(AbsoluteFileOffset + DurationPatchOffset, MakeConstArrayView(CompiledBoxData.GetData() + DurationPatchOffset, Version == 1 ? 8 : 4));
+}
+
+
+
+
+bool FBoxMDHD::BuildData(const FBoxBuilderData& InOutBuilderData)
+{
+	// Any of the values require a version 1 box?
+	Version = (CreationTime > 0xffffffffU || ModificationTime > 0xffffffffU || (Duration != ~((uint64)0) && Duration > 0xffffffffU)) ? 1 : 0;
+
+	FBoxFull::BuildData(InOutBuilderData);
+	FBoxDataWriter dw(CompiledBoxData);
+	if (Version == 1)
+	{
+		dw.PutAlignedValue(CreationTime);
+		dw.PutAlignedValue(ModificationTime);
+		dw.PutAlignedValue(Timescale);
+		DurationPatchOffset = dw.PutAlignedValue(Duration);
+	}
+	else
+	{
+		dw.PutAlignedValue((uint32)CreationTime);
+		dw.PutAlignedValue((uint32)ModificationTime);
+		dw.PutAlignedValue(Timescale);
+		DurationPatchOffset = dw.PutAlignedValue((uint32)Duration);
+	}
+	dw.PutBits(0, 1);
+	dw.PutBits(LanguageCode[0], 5);
+	dw.PutBits(LanguageCode[1], 5);
+	dw.PutBits(LanguageCode[2], 5);
+	check(dw.IsByteAligned());
+	dw.PutAlignedValue((uint16)0);
+	Size = dw.GetBytePosition();
+	bWasBuilt = true;
+	return true;
+}
+bool FBoxMDHD::UpdateDuration(uint64 InDuration)
+{
+	if (!bWasBuilt)
+	{
+		SetDuration(InDuration);
+		return true;
+	}
+	// If the duration now requires 64 bit wide fields when originally we had only 32 then this cannot be updated.
+	if (InDuration > 0xffffffffU && Version == 0)
+	{
+		return false;
+	}
+	SetDuration(InDuration);
+	FBoxDataWriter dw(CompiledBoxData);
+	if (Version == 1)
+	{
+		dw.SetAlignedValueAt(Duration, DurationPatchOffset);
+	}
+	else
+	{
+		dw.SetAlignedValueAt((uint32)Duration, DurationPatchOffset);
+	}
+	return true;
+}
+bool FBoxMDHD::RewriteChanges(FRewriteBoxChangeDelegate InRewriteBoxChangeDelegate)
+{
+	// If this box was not yet built then there's noting to rewrite.
+	if (!bWasBuilt)
+	{
+		return true;
+	}
+	return InRewriteBoxChangeDelegate.Execute(AbsoluteFileOffset + DurationPatchOffset, MakeConstArrayView(CompiledBoxData.GetData() + DurationPatchOffset, Version == 1 ? 8 : 4));
+}
+
+
+bool FBoxHDLR::BuildData(const FBoxBuilderData& InOutBuilderData)
+{
+	FBoxFull::BuildData(InOutBuilderData);
+	FBoxDataWriter dw(CompiledBoxData);
+	dw.PutAlignedValue((uint32)0);
+	dw.PutAlignedValue(HandlerType);
+	dw.PutAlignedValue((uint32)0);
+	dw.PutAlignedValue((uint32)0);
+	dw.PutAlignedValue((uint32)0);
+	FTCHARToUTF8 cnv(*Name);
+	int32 Len = cnv.Length();
+	auto Ptr = cnv.Get();
+	for(int32 i=0; i<Len; ++i)
+	{
+		dw.PutAlignedValue((uint8)*Ptr++);
+	}
+	dw.PutAlignedValue((uint8)0);
+	Size = dw.GetBytePosition();
+	return true;
+}
+
+
+bool FBoxVMHD::BuildData(const FBoxBuilderData& InOutBuilderData)
+{
+	Flags = 1;
+	FBoxFull::BuildData(InOutBuilderData);
+	FBoxDataWriter dw(CompiledBoxData);
+	dw.PutAlignedValue(GraphicsMode);
+	dw.PutAlignedValue(OpColor[0]);
+	dw.PutAlignedValue(OpColor[1]);
+	dw.PutAlignedValue(OpColor[2]);
+	Size = dw.GetBytePosition();
+	return true;
+}
+
+
+bool FBoxSMHD::BuildData(const FBoxBuilderData& InOutBuilderData)
+{
+	FBoxFull::BuildData(InOutBuilderData);
+	FBoxDataWriter dw(CompiledBoxData);
+	dw.PutAlignedValue(Balance);
+	dw.PutAlignedValue((uint16)0);
+	Size = dw.GetBytePosition();
+	return true;
+}
+
+
+bool FBoxDREF::BuildData(const FBoxBuilderData& InOutBuilderData)
+{
+	FBoxFull::BuildData(InOutBuilderData);
+	FBoxDataWriter dw(CompiledBoxData);
+	dw.PutAlignedValue((uint32)Children.Num());
+	Size = dw.GetBytePosition();
+	return true;
+}
+
+
+bool FBoxURL::BuildData(const FBoxBuilderData& InOutBuilderData)
+{
+	Flags = 1;
+	FBoxFull::BuildData(InOutBuilderData);
+	FBoxDataWriter dw(CompiledBoxData);
+	Size = dw.GetBytePosition();
+	return true;
+}
+
+
+bool FBoxSampleEntry::BuildData(const FBoxBuilderData& InOutBuilderData)
+{
+	FBoxBase::BuildData(InOutBuilderData);
+	FBoxDataWriter dw(CompiledBoxData);
+	dw.PutAlignedValue((uint8)0);				// reserved
+	dw.PutAlignedValue((uint8)0);				// reserved
+	dw.PutAlignedValue((uint8)0);				// reserved
+	dw.PutAlignedValue((uint8)0);				// reserved
+	dw.PutAlignedValue((uint8)0);				// reserved
+	dw.PutAlignedValue((uint8)0);				// reserved
+	dw.PutAlignedValue(DataReferenceIndex);		// data_reference_index
+	Size = dw.GetBytePosition();
+	return true;
+}
+
+bool FBoxVisualSampleEntry::BuildData(const FBoxBuilderData& InOutBuilderData)
+{
+	FBoxSampleEntry::BuildData(InOutBuilderData);
+	FBoxDataWriter dw(CompiledBoxData);
+	dw.PutAlignedValue((uint16)0);				// pre_defined
+	dw.PutAlignedValue((uint16)0);				// reserved
+	dw.PutAlignedValue((uint32)0);				// pre_defined
+	dw.PutAlignedValue((uint32)0);				// pre_defined
+	dw.PutAlignedValue((uint32)0);				// pre_defined
+	dw.PutAlignedValue(Width);
+	dw.PutAlignedValue(Height);
+	dw.PutAlignedValue((uint32)0x00480000U);	// horzresolution 72dpi
+	dw.PutAlignedValue((uint32)0x00480000U);	// vertresolution 72dpi
+	dw.PutAlignedValue((uint32)0);				// reserved
+	dw.PutAlignedValue((uint16)1);				// frame_count
+
+	FTCHARToUTF8 cnv(*CompressorName);
+	int32 Len = cnv.Length();
+	Len = Len > 31 ? 31 : Len;
+	dw.PutAlignedValue((uint8)Len);
+	auto Ptr = cnv.Get();
+	int32 i;
+	for(i=0; i<Len; ++i)
+	{
+		dw.PutAlignedValue((uint8)*Ptr++);
+	}
+	for(; i<31; ++i)
+	{
+		dw.PutAlignedValue((uint8)0);
+	}
+	dw.PutAlignedValue((uint16)0x18);			// depth
+	dw.PutAlignedValue((uint16)0xffff);			// pre_defined
+	for(i=0; i<AdditionalBoxes.Num(); ++i)
+	{
+		dw.PutAlignedArray(AdditionalBoxes[i]);
+	}
+	Size = dw.GetBytePosition();
+	return true;
+}
+
+
+bool FBoxAudioSampleEntry::BuildData(const FBoxBuilderData& InOutBuilderData)
+{
+	FBoxSampleEntry::BuildData(InOutBuilderData);
+	FBoxDataWriter dw(CompiledBoxData);
+	// Need an `AudioSampleEntryV1` instead?
+	const bool bV1 = SampleRate >= 0x10000;
+	dw.PutAlignedValue((uint16)(bV1 ? 1U : 0U));// version
+	dw.PutAlignedValue((uint16)0);				// reserved
+	dw.PutAlignedValue((uint32)0);				// reserved
+	dw.PutAlignedValue(ChannelCount);
+	dw.PutAlignedValue((uint16)16);				// samplesize == 16
+	dw.PutAlignedValue((uint16)0);				// reserved
+	dw.PutAlignedValue((uint16)0);				// reserved
+	dw.PutAlignedValue((uint32)(bV1 ? 0x10000U : (SampleRate << 16)));
+	// Add an `srat` box, which is full box version=0,flags=0
+	if (bV1)
+	{
+		dw.PutAlignedValue((uint32)16U);
+		dw.PutAlignedValue(BoxAtom('s','r','a','t'));
+		dw.PutAlignedValue((uint32)0U);
+		dw.PutAlignedValue(SampleRate);
+	}
+	for(int32 i=0; i<AdditionalBoxes.Num(); ++i)
+	{
+		dw.PutAlignedArray(AdditionalBoxes[i]);
+	}
+	Size = dw.GetBytePosition();
+	return true;
+}
+
+
+bool FBoxSTSD::BuildData(const FBoxBuilderData& InOutBuilderData)
+{
+	FBoxFull::BuildData(InOutBuilderData);
+	FBoxDataWriter dw(CompiledBoxData);
+	dw.PutAlignedValue((uint32)Children.Num());
+	Size = dw.GetBytePosition();
+	return true;
+}
+
+
+bool FBoxSTTS::BuildData(const FBoxBuilderData& InOutBuilderData)
+{
+	FBoxFull::BuildData(InOutBuilderData);
+	FBoxDataWriter dw(CompiledBoxData);
+	dw.PutAlignedValue((uint32)SampleDurations.Num());
+	for(int32 i=0, iMax=SampleDurations.Num(); i<iMax; ++i)
+	{
+		const FEntry& e(SampleDurations[i]);
+		dw.PutAlignedValue(e.Count);
+		dw.PutAlignedValue(e.Delta);
+	}
+	Size = dw.GetBytePosition();
+	return true;
+}
+
+
+bool FBoxSTSS::BuildData(const FBoxBuilderData& InOutBuilderData)
+{
+	FBoxFull::BuildData(InOutBuilderData);
+	FBoxDataWriter dw(CompiledBoxData);
+	dw.PutAlignedValue((uint32)SyncSampleIndices.Num());
+	for(int32 i=0, iMax=SyncSampleIndices.Num(); i<iMax; ++i)
+	{
+		dw.PutAlignedValue(SyncSampleIndices[i]);
+	}
+	Size = dw.GetBytePosition();
+	return true;
+}
+
+
+bool FBoxCTTS::BuildData(const FBoxBuilderData& InOutBuilderData)
+{
+	Version = bHasNegativeEntries ? 1 : 0;
+	FBoxFull::BuildData(InOutBuilderData);
+	FBoxDataWriter dw(CompiledBoxData);
+	dw.PutAlignedValue((uint32)CompositionTimeOffsets.Num());
+	if (Version == 0)
+	{
+		for(int32 i=0, iMax=CompositionTimeOffsets.Num(); i<iMax; ++i)
+		{
+			const FEntry& e(CompositionTimeOffsets[i]);
+			dw.PutAlignedValue(e.Count);
+			dw.PutAlignedValue((uint32)e.Delta);
+		}
+	}
+	else
+	{
+		for(int32 i=0, iMax=CompositionTimeOffsets.Num(); i<iMax; ++i)
+		{
+			const FEntry& e(CompositionTimeOffsets[i]);
+			dw.PutAlignedValue(e.Count);
+			dw.PutAlignedValue(e.Delta);
+		}
+	}
+	Size = dw.GetBytePosition();
+	return true;
+}
+
+
+bool FBoxSTSC::BuildData(const FBoxBuilderData& InOutBuilderData)
+{
+	FBoxFull::BuildData(InOutBuilderData);
+	FBoxDataWriter dw(CompiledBoxData);
+	dw.PutAlignedValue((uint32)SampleToChunkEntries.Num());
+	for(int32 i=0, iMax=SampleToChunkEntries.Num(); i<iMax; ++i)
+	{
+		const FEntry& e(SampleToChunkEntries[i]);
+		dw.PutAlignedValue(e.first_chunk);
+		dw.PutAlignedValue(e.samples_per_chunk);
+		dw.PutAlignedValue(e.sample_description_index);
+	}
+	Size = dw.GetBytePosition();
+	return true;
+}
+
+
+bool FBoxSTSZ::BuildData(const FBoxBuilderData& InOutBuilderData)
+{
+	if (bUseSmallSizeBox && ConstantSampleSize == 0 && LargestSampleSize < 0x10000)
+	{
+		Type = BoxAtom('s','t','z','2');
+		uint8 NumBits = LargestSampleSize >= 256 ? 16 : LargestSampleSize >= 16 ? 8 : 4;
+		FBoxFull::BuildData(InOutBuilderData);
+		FBoxDataWriter dw(CompiledBoxData);
+		dw.PutAlignedValue((uint32)NumBits);
+		dw.PutAlignedValue((uint32)SampleSizes.Num());
+		if (NumBits == 16)
+		{
+			for(int32 i=0,iMax=SampleSizes.Num(); i<iMax; ++i)
+			{
+				dw.PutAlignedValue((uint16) SampleSizes[i]);
+			}
+		}
+		else if (NumBits == 8)
+		{
+			for(int32 i=0,iMax=SampleSizes.Num(); i<iMax; ++i)
+			{
+				dw.PutAlignedValue((uint8) SampleSizes[i]);
+			}
+		}
+		else// if (NumBits == 4)
+		{
+			for(int32 i=0,iMax=SampleSizes.Num() & ~1; i<iMax; i+=2)
+			{
+				uint8 v = ((SampleSizes[i+1] & 15) << 4) | (SampleSizes[i] & 15);
+				dw.PutAlignedValue(v);
+			}
+			if (SampleSizes.Num() & 1)
+			{
+				uint8 v = (SampleSizes[SampleSizes.Num()-1] & 15) << 4;
+				dw.PutAlignedValue(v);
+			}
+		}
+		Size = dw.GetBytePosition();
+	}
+	else
+	{
+		FBoxFull::BuildData(InOutBuilderData);
+		FBoxDataWriter dw(CompiledBoxData);
+		dw.PutAlignedValue(ConstantSampleSize);
+		if (ConstantSampleSize == 0)
+		{
+			dw.PutAlignedValue((uint32)SampleSizes.Num());
+			for(int32 i=0,iMax=SampleSizes.Num(); i<iMax; ++i)
+			{
+				dw.PutAlignedValue(SampleSizes[i]);
+			}
+		}
+		Size = dw.GetBytePosition();
+	}
+	return true;
+}
+
+
+bool FBoxSTCO::BuildData(const FBoxBuilderData& InOutBuilderData)
+{
+	if (bNeeds64Bits)
+	{
+		Type = BoxAtom('c','o','6','4');
+	}
+	FBoxFull::BuildData(InOutBuilderData);
+	FBoxDataWriter dw(CompiledBoxData);
+	dw.PutAlignedValue((uint32)ChunkOffsets.Num());
+	if (bNeeds64Bits)
+	{
+		for(int32 i=0,iMax=ChunkOffsets.Num(); i<iMax; ++i)
+		{
+			dw.PutAlignedValue((uint64)ChunkOffsets[i]);
+		}
+	}
+	else
+	{
+		for(int32 i=0,iMax=ChunkOffsets.Num(); i<iMax; ++i)
+		{
+			dw.PutAlignedValue((uint32)ChunkOffsets[i]);
+		}
+	}
+	Size = dw.GetBytePosition();
+	return true;
+}
+
+
+bool FBoxMEHD::BuildData(const FBoxBuilderData& InOutBuilderData)
+{
+	Version = FragmentDuration > 0xffffffffU ? 1 : 0;
+	FBoxFull::BuildData(InOutBuilderData);
+	FBoxDataWriter dw(CompiledBoxData);
+	if (Version == 1)
+	{
+		FragmentDurationPatchOffset = dw.PutAlignedValue((uint64) FragmentDuration);
+	}
+	else
+	{
+		FragmentDurationPatchOffset = dw.PutAlignedValue((uint32) FragmentDuration);
+	}
+	Size = dw.GetBytePosition();
+	bWasBuilt = true;
+	return true;
+}
+bool FBoxMEHD::UpdateFragmentDuration(uint64 InFragmentDuration)
+{
+	if (!bWasBuilt)
+	{
+		SetFragmentDuration(InFragmentDuration);
+		return true;
+	}
+	// If the duration now requires 64 bit wide fields when originally we had only 32 then this cannot be updated.
+	if (InFragmentDuration > 0xffffffffU && Version == 0)
+	{
+		return false;
+	}
+	SetFragmentDuration(InFragmentDuration);
+	FBoxDataWriter dw(CompiledBoxData);
+	if (Version == 1)
+	{
+		dw.SetAlignedValueAt(FragmentDuration, FragmentDurationPatchOffset);
+	}
+	else
+	{
+		dw.SetAlignedValueAt((uint32)FragmentDuration, FragmentDurationPatchOffset);
+	}
+	return true;
+}
+bool FBoxMEHD::RewriteChanges(FRewriteBoxChangeDelegate InRewriteBoxChangeDelegate)
+{
+	// If this box was not yet built then there's noting to rewrite.
+	if (!bWasBuilt)
+	{
+		return true;
+	}
+	return InRewriteBoxChangeDelegate.Execute(AbsoluteFileOffset + FragmentDurationPatchOffset, MakeConstArrayView(CompiledBoxData.GetData() + FragmentDurationPatchOffset, Version == 1 ? 8 : 4));
+}
+
+
+bool FBoxTREX::BuildData(const FBoxBuilderData& InOutBuilderData)
+{
+	FBoxFull::BuildData(InOutBuilderData);
+	FBoxDataWriter dw(CompiledBoxData);
+	dw.PutAlignedValue(TrackID);
+	dw.PutAlignedValue(default_sample_description_index);
+	dw.PutAlignedValue(default_sample_duration);
+	dw.PutAlignedValue(default_sample_size);
+	dw.PutAlignedValue(default_sample_flags);
+	Size = dw.GetBytePosition();
+	return true;
+}
+
+
+bool FBoxMFHD::BuildData(const FBoxBuilderData& InOutBuilderData)
+{
+	FBoxFull::BuildData(InOutBuilderData);
+	FBoxDataWriter dw(CompiledBoxData);
+	dw.PutAlignedValue(SequenceNumber);
+	Size = dw.GetBytePosition();
+	return true;
+}
+
+
+bool FBoxTFHD::BuildData(const FBoxBuilderData& InOutBuilderData)
+{
+	if (base_data_offset.IsSet())
+	{
+		Flags |= 0x1;
+	}
+	else if (bDefaultBaseIsMoof)
+	{
+		Flags |= 0x20000;
+	}
+	if (sample_description_index.IsSet())
+	{
+		Flags |= 0x2;
+	}
+	if (default_sample_duration.IsSet())
+	{
+		Flags |= 0x8;
+	}
+	if (default_sample_size.IsSet())
+	{
+		Flags |= 0x10;
+	}
+	if (default_sample_flags.IsSet())
+	{
+		Flags |= 0x20;
+	}
+
+	FBoxFull::BuildData(InOutBuilderData);
+	FBoxDataWriter dw(CompiledBoxData);
+	dw.PutAlignedValue(TrackID);
+	if (base_data_offset.IsSet())
+	{
+		BaseDataPatchOffset = dw.PutAlignedValue(base_data_offset.GetValue());
+	}
+	if (sample_description_index.IsSet())
+	{
+		dw.PutAlignedValue(sample_description_index.GetValue());
+	}
+	if (default_sample_duration.IsSet())
+	{
+		dw.PutAlignedValue(default_sample_duration.GetValue());
+	}
+	if (default_sample_size.IsSet())
+	{
+		dw.PutAlignedValue(default_sample_size.GetValue());
+	}
+	if (default_sample_flags.IsSet())
+	{
+		dw.PutAlignedValue(default_sample_flags.GetValue());
+	}
+	Size = dw.GetBytePosition();
+	bWasBuilt = true;
+	return true;
+}
+bool FBoxTFHD::UpdateBaseDataOffset(uint64 InBaseDataOffset)
+{
+	if (!bWasBuilt)
+	{
+		SetBaseDataOffset(InBaseDataOffset);
+		return true;
+	}
+	// If the base data offset was not set before it cannot be set now.
+	if (!base_data_offset.IsSet())
+	{
+		return false;
+	}
+	SetBaseDataOffset(InBaseDataOffset);
+	FBoxDataWriter dw(CompiledBoxData);
+	dw.SetAlignedValueAt(base_data_offset.GetValue(), BaseDataPatchOffset);
+	return true;
+}
+bool FBoxTFHD::RewriteChanges(FRewriteBoxChangeDelegate InRewriteBoxChangeDelegate)
+{
+	// If this box was not yet built then there's noting to rewrite.
+	if (!bWasBuilt)
+	{
+		return true;
+	}
+	return InRewriteBoxChangeDelegate.Execute(AbsoluteFileOffset + BaseDataPatchOffset, MakeConstArrayView(CompiledBoxData.GetData() + BaseDataPatchOffset, 8));
+}
+
+
+bool FBoxTFDT::BuildData(const FBoxBuilderData& InOutBuilderData)
+{
+	Version = BaseMediaDecodeTime > 0xffffffffU ? 1 : 0;
+	FBoxFull::BuildData(InOutBuilderData);
+	FBoxDataWriter dw(CompiledBoxData);
+	if (Version == 1)
+	{
+		dw.PutAlignedValue(BaseMediaDecodeTime);
+	}
+	else
+	{
+		dw.PutAlignedValue((uint32)BaseMediaDecodeTime);
+	}
+	Size = dw.GetBytePosition();
+	return true;
+}
+
+
+bool FBoxTRUN::BuildData(const FBoxBuilderData& InOutBuilderData)
+{
+	Version = bRequiresCompositionTimeOffsets ? 1 : 0;
+	if (DataOffset.IsSet())
+	{
+		Flags |= 0x1;
+	}
+	// If the first sample has dedicated flags then no flags for individual samples must be used!
+	// See ISO/IEC 14496-12:2022 - 8.8.8.1 Definition
+	if (first_sample_flags.IsSet())
+	{
+		Flags |= 0x4;
+	}
+	else if (!bDoNotWriteFlags)
+	{
+		Flags |= 0x400;
+	}
+
+	if (bDoNotWriteDurations == false && (bHasDifferentDurations || Samples.Num() == 1))
+	{
+		Flags |= 0x100;
+	}
+	if (bDoNotWriteSizes == false && (bHasDifferentSizes || Samples.Num() == 1))
+	{
+		Flags |= 0x200;
+	}
+	if (bRequiresCompositionTimeOffsets)
+	{
+		Flags |= 0x800;
+	}
+
+	FBoxFull::BuildData(InOutBuilderData);
+	FBoxDataWriter dw(CompiledBoxData);
+	dw.PutAlignedValue((uint32)Samples.Num());
+	if ((Flags & 0x1) != 0)
+	{
+		DataOffsetPatchOffset = dw.PutAlignedValue(DataOffset.GetValue());
+	}
+	if ((Flags & 0x4) != 0)
+	{
+		dw.PutAlignedValue(first_sample_flags.GetValue());
+	}
+	// Only write sample information if any of them are to be.
+	// If everything uses default values then there is no need for that.
+	if ((Flags & 0xf00) != 0)
+	{
+		for(int32 i=0, iMax=Samples.Num(); i<iMax; ++i)
+		{
+			const FEntry& e(Samples[i]);
+			if (Flags & 0x100)
+			{
+				dw.PutAlignedValue(e.sample_duration);
+			}
+			if (Flags & 0x200)
+			{
+				dw.PutAlignedValue(e.sample_size);
+			}
+			if (Flags & 0x400)
+			{
+				dw.PutAlignedValue(e.sample_flags);
+			}
+			if (Flags & 0x800)
+			{
+				dw.PutAlignedValue(e.sample_composition_time_offset);
+			}
+		}
+	}
+	Size = dw.GetBytePosition();
+	return true;
+}
+
+
+bool FBoxSUBS::BuildData(const FBoxBuilderData& InOutBuilderData)
+{
+	Version = 1;	// always write 32 bit subsample sizes
+	FBoxFull::BuildData(InOutBuilderData);
+	FBoxDataWriter dw(CompiledBoxData);
+	dw.PutAlignedValue((uint32)SubSamples.Num());
+	uint32 PreviousSampleNum = 0;
+	for(int32 i=0, iMax=SubSamples.Num(); i<iMax; ++i)
+	{
+		uint32 SampleDelta = SubSamples[i].SampleNumber - PreviousSampleNum;
+		PreviousSampleNum = SubSamples[i].SampleNumber;
+		dw.PutAlignedValue(SampleDelta);
+		dw.PutAlignedValue((uint16) SubSamples[i].Subs.Num());
+		for(int32 j=0, jMax=SubSamples[i].Subs.Num(); j<jMax; ++j)
+		{
+			if (Version == 1)
+			{
+				dw.PutAlignedValue((uint32) SubSamples[i].Subs[j].subsample_size);
+			}
+			else
+			{
+				dw.PutAlignedValue((uint16) SubSamples[i].Subs[j].subsample_size);
+			}
+			dw.PutAlignedValue((uint8) SubSamples[i].Subs[j].subsample_priority);
+			dw.PutAlignedValue((uint8) SubSamples[i].Subs[j].discardable);
+			dw.PutAlignedValue((uint32) SubSamples[i].Subs[j].codec_specific_parameters);
+		}
+	}
+	Size = dw.GetBytePosition();
+	return true;
+}
+
+
+bool FBoxELST::BuildData(const FBoxBuilderData& InOutBuilderData)
+{
+	Version = EditDuration > TNumericLimits<uint32>::Max() || MediaTime > TNumericLimits<int32>::Max() || MediaTime < TNumericLimits<int32>::Min() ? 1 : 0;
+	FBoxFull::BuildData(InOutBuilderData);
+	FBoxDataWriter dw(CompiledBoxData);
+	dw.PutAlignedValue((uint32)1);	// just one entry.
+	if (Version == 1)
+	{
+		EditDurationPatchOffset = dw.PutAlignedValue(EditDuration);
+		dw.PutAlignedValue(MediaTime);
+	}
+	else
+	{
+		EditDurationPatchOffset = dw.PutAlignedValue((uint32)EditDuration);
+		dw.PutAlignedValue((int32)MediaTime);
+	}
+	dw.PutAlignedValue((int16)1);
+	dw.PutAlignedValue((int16)0);
+	Size = dw.GetBytePosition();
+	bWasBuilt = true;
+	return true;
+}
+
+bool FBoxELST::UpdateEditDuration(uint64 InEditDuration)
+{
+	if (!bWasBuilt)
+	{
+		SetEditDuration(InEditDuration);
+		return true;
+	}
+	if (InEditDuration > TNumericLimits<uint32>::Max() && Version == 0)
+	{
+		return false;
+	}
+	SetEditDuration(InEditDuration);
+	FBoxDataWriter dw(CompiledBoxData);
+	if (Version == 1)
+	{
+		dw.SetAlignedValueAt(EditDuration, EditDurationPatchOffset);
+	}
+	else
+	{
+		dw.SetAlignedValueAt((uint32)EditDuration, EditDurationPatchOffset);
+	}
+	return true;
+}
+bool FBoxELST::RewriteChanges(FRewriteBoxChangeDelegate InRewriteBoxChangeDelegate)
+{
+	// If this box was not yet built then there's noting to rewrite.
+	if (!bWasBuilt)
+	{
+		return true;
+	}
+	return InRewriteBoxChangeDelegate.Execute(AbsoluteFileOffset + EditDurationPatchOffset, MakeConstArrayView(CompiledBoxData.GetData() + EditDurationPatchOffset, Version == 1 ? 8 : 4));
+}
+
+
+}
diff --git a/Engine/Plugins/Media/MP4Utilities/Source/MP4Muxer/Private/MP4MuxerBoxes.h b/Engine/Plugins/Media/MP4Utilities/Source/MP4Muxer/Private/MP4MuxerBoxes.h
new file mode 100644
--- /dev/null
+++ b/Engine/Plugins/Media/MP4Utilities/Source/MP4Muxer/Private/MP4MuxerBoxes.h
@@ -0,0 +1,942 @@
+// Copyright Epic Games, Inc. All Rights Reserved.
+
+#pragma once
+
+#include "CoreTypes.h"
+#include "Containers/Array.h"
+#include "Containers/ChunkedArray.h"
+#include "MP4Muxer.h"
+
+namespace MP4MuxerBoxes
+{
+static constexpr uint16 EndianSwap(uint16 value)			{ return (value >> 8) | (value << 8); }
+static constexpr int16 EndianSwap(int16 value)				{ return int16(EndianSwap(uint16(value))); }
+static constexpr uint32 EndianSwap(uint32 value)			{ return (value << 24) | ((value & 0xff00) << 8) | ((value >> 8) & 0xff00) | (value >> 24); }
+static constexpr int32 EndianSwap(int32 value)				{ return int32(EndianSwap(uint32(value))); }
+static constexpr uint64 EndianSwap(uint64 value)			{ return (uint64(EndianSwap(uint32(value & 0xffffffffU))) << 32) | uint64(EndianSwap(uint32(value >> 32))); }
+static constexpr int64 EndianSwap(int64 value)				{ return int64(EndianSwap(uint64(value)));}
+
+static constexpr uint32 BoxAtom(const uint8 A, const uint8 B, const uint8 C, const uint8 D)
+{ return (static_cast<uint32>(A) << 24) | (static_cast<uint32>(B) << 16) | (static_cast<uint32>(C) << 8) | static_cast<uint32>(D); }
+
+
+class FBoxBuilderData
+{
+public:
+	FBoxBuilderData() = default;
+	~FBoxBuilderData() = default;
+};
+
+class FBoxBase
+{
+public:
+	FBoxBase(uint32 InType) : Type(InType)
+	{ }
+	virtual ~FBoxBase();
+	virtual bool BuildData(const FBoxBuilderData& InOutBuilderData);
+	virtual void CalculateBoxSizes();
+	void AddChild(FBoxBase* InChild)
+	{
+		InChild->Parent = this;
+		Children.Add(InChild);
+	}
+	void AddChildAfter(FBoxBase* InChild, FBoxBase* InAfter)
+	{
+		InChild->Parent = this;
+		int32 Pos = Children.IndexOfByKey(InAfter);
+		if (Pos != INDEX_NONE && Pos+1<Children.Num())
+		{
+			Children.Insert(InChild, Pos+1);
+		}
+		else
+		{
+			Children.Add(InChild);
+		}
+	}
+	FBoxBase* GetParent() const
+	{ return Parent; }
+	void RemoveChild(FBoxBase* InChild)
+	{ Children.Remove(InChild); }
+	template <typename Predicate>
+	bool VisitForWriting(Predicate Pred)
+	{
+		if (!::Invoke(Pred, this))
+		{
+			return false;
+		}
+		for(int32 i=0; i<Children.Num(); ++i)
+		{
+			if (!Children[i]->VisitForWriting(Pred))
+			{
+				return false;
+			}
+		}
+		return true;
+	}
+	const TArray<uint8>& GetCompiledBoxData() const
+	{ return CompiledBoxData; }
+	// Gets set while writing
+	void SetAbsoluteFileOffset(int64 InAbsoluteFileOffset)
+	{ AbsoluteFileOffset = InAbsoluteFileOffset; }
+	int64 GetAbsoluteFileOffset() const
+	{ return AbsoluteFileOffset; }
+
+	DECLARE_DELEGATE_RetVal_TwoParams(bool, FRewriteBoxChangeDelegate, int64 /*FilePosition*/, TConstArrayView<uint8> /*NewData*/);
+	virtual bool RewriteChanges(FRewriteBoxChangeDelegate InRewriteBoxChangeDelegate)
+	{ return true; }
+
+protected:
+	void SetType(uint32 InType)
+	{ Type = InType; }
+	void SetSize(uint32 InSize)
+	{ Size = InSize; }
+	void AddSizeOfChildren();
+	uint32 Type = 0;
+	uint32 Size = 0;
+	FBoxBase* Parent = nullptr;
+	TArray<FBoxBase*> Children;
+	TArray<uint8> CompiledBoxData;
+	uint32 CompiledBoxData_Offset_Size = 0;
+	int64 AbsoluteFileOffset = -1;
+};
+
+class FBoxFull : public FBoxBase
+{
+public:
+	FBoxFull(uint32 InType) : FBoxBase(InType)
+	{ }
+	virtual ~FBoxFull() = default;
+	bool BuildData(const FBoxBuilderData& InOutBuilderData) override;
+	void SetFlags(uint32 InFlags)
+	{ Flags = InFlags & 0x00ffffff; }
+protected:
+	uint8 Version = 0;
+	uint32 Flags = 0;
+};
+
+
+class FBoxFTYP : public FBoxBase
+{
+public:
+	FBoxFTYP() : FBoxBase(BoxAtom('f','t','y','p'))
+	{ }
+	virtual ~FBoxFTYP() = default;
+
+	void SetMajorBrand(uint32 InMajorBrand)
+	{ MajorBrand = InMajorBrand; }
+	void SetMinorVersion(uint32 InMinorVersion)
+	{ MinorVersion = InMinorVersion; }
+	void AddCompatibleBrand(uint32 InCompatibleBrand)
+	{ CompatibleBrands.Emplace(InCompatibleBrand); }
+
+	bool BuildData(const FBoxBuilderData& InOutBuilderData) override;
+	bool UpdateMajorBrand(uint32 InMajorBrand);
+	bool RewriteChanges(FRewriteBoxChangeDelegate InRewriteBoxChangeDelegate) override;
+protected:
+	uint32 MajorBrand = 0;
+	uint32 MinorVersion = 0;
+	TArray<uint32> CompatibleBrands;
+	int64 MajorBrandPatchOffset = -1;
+	bool bWasBuilt = false;
+};
+
+
+class FBoxFREE : public FBoxBase
+{
+public:
+	FBoxFREE() : FBoxBase(BoxAtom('f','r','e','e'))
+	{ }
+	virtual ~FBoxFREE() = default;
+	void SetNumberOfEmptyBytes(uint32 InNumEmptyBytes)
+	{ NumEmptyBytes = InNumEmptyBytes; }
+	bool BuildData(const FBoxBuilderData& InOutBuilderData) override;
+protected:
+	uint32 NumEmptyBytes = 0;
+};
+
+
+class FBoxMDAT : public FBoxBase
+{
+public:
+	FBoxMDAT() : FBoxBase(BoxAtom('m','d','a','t'))
+	{ }
+	virtual ~FBoxMDAT() = default;
+protected:
+};
+
+
+class FBoxMOOV : public FBoxBase
+{
+public:
+	FBoxMOOV() : FBoxBase(BoxAtom('m','o','o','v'))
+	{ }
+	virtual ~FBoxMOOV() = default;
+protected:
+};
+
+
+class FBoxMVHD : public FBoxFull
+{
+public:
+	FBoxMVHD() : FBoxFull(BoxAtom('m','v','h','d'))
+	{ }
+	virtual ~FBoxMVHD() = default;
+
+	void SetCreationTime(uint64 InCreationTime)
+	{ CreationTime = InCreationTime; }
+	void SetModificationTime(uint64 InModificationTime)
+	{ ModificationTime = InModificationTime; }
+	void SetTimescale(uint32 InTimescale)
+	{ Timescale = InTimescale; }
+	uint32 GetTimescale() const
+	{ return Timescale; }
+	void SetDuration(uint64 InDuration)
+	{ Duration = InDuration; }
+	uint64 GetDuration() const
+	{ return Duration; }
+	void SetNextTrackID(uint32 InNextTrackID)
+	{ NextTrackID = InNextTrackID; }
+
+	bool BuildData(const FBoxBuilderData& InOutBuilderData) override;
+	bool UpdateTimescale(uint32 InTimescale);
+	bool UpdateDuration(uint64 InDuration);
+	bool RewriteChanges(FRewriteBoxChangeDelegate InRewriteBoxChangeDelegate) override;
+protected:
+	uint64 CreationTime = 0;
+	uint64 ModificationTime = 0;
+	uint64 Duration = 0;
+	uint32 Timescale = 0;
+	uint32 NextTrackID = 0;
+	int64 TimescalePatchOffset = -1;
+	int64 DurationPatchOffset = -1;
+	bool bWasBuilt = false;
+};
+
+
+class FBoxTRAK : public FBoxBase
+{
+public:
+	FBoxTRAK() : FBoxBase(BoxAtom('t','r','a','k'))
+	{ }
+	virtual ~FBoxTRAK() = default;
+protected:
+};
+
+
+class FBoxTKHD : public FBoxFull
+{
+public:
+	FBoxTKHD() : FBoxFull(BoxAtom('t','k','h','d'))
+	{ }
+	virtual ~FBoxTKHD() = default;
+
+	void SetCreationTime(uint64 InCreationTime)
+	{ CreationTime = InCreationTime; }
+	void SetModificationTime(uint64 InModificationTime)
+	{ ModificationTime = InModificationTime; }
+	void SetDuration(uint64 InDuration)
+	{ Duration = InDuration; }
+	void SetTrackID(uint32 InTrackID)
+	{ TrackID = InTrackID; }
+	void SetTrackEnabled(bool bIsEnabled)
+	{ if (bIsEnabled) Flags |= 1U; else Flags &= ~(1U); }
+	void SetTrackInMovie(bool bIsInMovie)
+	{ if (bIsInMovie) Flags |= 2U; else Flags &= ~(2U); }
+	void SetTrackInPreview(bool bIsInPreview)
+	{ if (bIsInPreview) Flags |= 4U; else Flags &= ~(4U); }
+	void SetWidth(uint32 InWidth)
+	{ Width = InWidth << 16; }
+	void SetHeight(uint32 InHeight)
+	{ Height = InHeight << 16; }
+	void SetVolume(uint8 InVolume)
+	{ Volume = (uint16)InVolume << 8; }
+	void SetLayer(uint16 InLayer)
+	{ Layer = InLayer; }
+	void SetAlternateGroup(uint16 InAlternateGroup)
+	{ AlternateGroup = InAlternateGroup; }
+
+	bool BuildData(const FBoxBuilderData& InOutBuilderData) override;
+
+	bool UpdateDuration(uint64 InDuration);
+	bool RewriteChanges(FRewriteBoxChangeDelegate InRewriteBoxChangeDelegate) override;
+protected:
+	uint64 CreationTime = 0;
+	uint64 ModificationTime = 0;
+	uint64 Duration = 0;
+	uint32 TrackID = 0;
+	uint32 Width = 0;
+	uint32 Height = 0;
+	uint16 Volume = 0;
+	uint16 Layer = 0;
+	uint16 AlternateGroup = 0;
+	int64 DurationPatchOffset = -1;
+	bool bWasBuilt = false;
+};
+
+
+class FBoxMDIA : public FBoxBase
+{
+public:
+	FBoxMDIA() : FBoxBase(BoxAtom('m','d','i','a'))
+	{ }
+	virtual ~FBoxMDIA() = default;
+protected:
+};
+
+
+class FBoxMDHD : public FBoxFull
+{
+public:
+	FBoxMDHD() : FBoxFull(BoxAtom('m','d','h','d'))
+	{ }
+	virtual ~FBoxMDHD() = default;
+
+	void SetCreationTime(uint64 InCreationTime)
+	{ CreationTime = InCreationTime; }
+	void SetModificationTime(uint64 InModificationTime)
+	{ ModificationTime = InModificationTime; }
+	void SetDuration(uint64 InDuration)
+	{ Duration = InDuration; }
+	void SetTimescale(uint32 InTimescale)
+	{ Timescale = InTimescale; }
+	void SetLanguage(uint8 In1, uint8 In2, uint8 In3)
+	{
+		check(In1 >= 0x60 && In1 < 0x80 && In2 >= 0x60 && In2 < 0x80 && In3 >= 0x60 && In3 < 0x80);
+		LanguageCode[0] = In1 - 0x60;
+		LanguageCode[1] = In2 - 0x60;
+		LanguageCode[2] = In3 - 0x60;
+	}
+
+	bool BuildData(const FBoxBuilderData& InOutBuilderData) override;
+	bool UpdateDuration(uint64 InDuration);
+	bool RewriteChanges(FRewriteBoxChangeDelegate InRewriteBoxChangeDelegate) override;
+protected:
+	uint64 CreationTime = 0;
+	uint64 ModificationTime = 0;
+	uint64 Duration = 0;
+	uint32 Timescale = 0;
+	uint8 LanguageCode[3] {};
+	int64 DurationPatchOffset = -1;
+	bool bWasBuilt = false;
+};
+
+
+class FBoxHDLR : public FBoxFull
+{
+public:
+	FBoxHDLR() : FBoxFull(BoxAtom('h','d','l','r'))
+	{ }
+	virtual ~FBoxHDLR() = default;
+
+	void SetHandlerType(uint32 InHandlerType)
+	{ HandlerType = InHandlerType; }
+	void SetName(const FString& InName)
+	{ Name = InName; }
+
+	bool BuildData(const FBoxBuilderData& InOutBuilderData) override;
+protected:
+	FString Name;
+	uint32 HandlerType = 0;
+};
+
+
+class FBoxMINF : public FBoxBase
+{
+public:
+	FBoxMINF() : FBoxBase(BoxAtom('m','i','n','f'))
+	{ }
+	virtual ~FBoxMINF() = default;
+protected:
+};
+
+
+class FBoxVMHD : public FBoxFull
+{
+public:
+	FBoxVMHD() : FBoxFull(BoxAtom('v','m','h','d'))
+	{ }
+	virtual ~FBoxVMHD() = default;
+
+	bool BuildData(const FBoxBuilderData& InOutBuilderData) override;
+protected:
+	uint16 GraphicsMode = 0;
+	uint16 OpColor[3] { };
+};
+
+
+class FBoxSMHD : public FBoxFull
+{
+public:
+	FBoxSMHD() : FBoxFull(BoxAtom('s','m','h','d'))
+	{ }
+	virtual ~FBoxSMHD() = default;
+
+	bool BuildData(const FBoxBuilderData& InOutBuilderData) override;
+protected:
+	uint16 Balance = 0;
+};
+
+
+class FBoxDINF : public FBoxBase
+{
+public:
+	FBoxDINF() : FBoxBase(BoxAtom('d','i','n','f'))
+	{ }
+	virtual ~FBoxDINF() = default;
+protected:
+};
+
+
+class FBoxDREF : public FBoxFull
+{
+public:
+	FBoxDREF() : FBoxFull(BoxAtom('d','r','e','f'))
+	{ }
+	virtual ~FBoxDREF() = default;
+	bool BuildData(const FBoxBuilderData& InOutBuilderData) override;
+protected:
+};
+
+
+class FBoxURL : public FBoxFull
+{
+public:
+	FBoxURL() : FBoxFull(BoxAtom('u','r','l',' '))
+	{ }
+	virtual ~FBoxURL() = default;
+	bool BuildData(const FBoxBuilderData& InOutBuilderData) override;
+protected:
+};
+
+
+class FBoxSTBL : public FBoxBase
+{
+public:
+	FBoxSTBL() : FBoxBase(BoxAtom('s','t','b','l'))
+	{ }
+	virtual ~FBoxSTBL() = default;
+protected:
+};
+
+
+class FBoxSampleEntry : public FBoxBase
+{
+public:
+	FBoxSampleEntry(uint32 InFormat) : FBoxBase(InFormat)
+	{ }
+	virtual ~FBoxSampleEntry() = default;
+	void SetAdditionalBoxes(const TArray<TArray<uint8>>& InAdditionalBoxes)
+	{ AdditionalBoxes = InAdditionalBoxes; }
+	bool BuildData(const FBoxBuilderData& InOutBuilderData) override;
+protected:
+	bool AppendAdditionalBoxes();
+	TArray<TArray<uint8>> AdditionalBoxes;
+	uint16 DataReferenceIndex = 1;
+};
+
+
+class FBoxVisualSampleEntry : public FBoxSampleEntry
+{
+public:
+	FBoxVisualSampleEntry(uint32 InFormat) : FBoxSampleEntry(InFormat)
+	{ }
+	virtual ~FBoxVisualSampleEntry() = default;
+	void SetWidth(uint16 InWidth)
+	{ Width = InWidth; }
+	void SetHeight(uint16 InHeight)
+	{ Height = InHeight; }
+	void SetCompressorName(const FString& InCompressorName)
+	{ CompressorName = InCompressorName; }
+	bool BuildData(const FBoxBuilderData& InOutBuilderData) override;
+protected:
+	FString CompressorName;
+	uint16 Width = 0;
+	uint16 Height = 0;
+};
+
+
+class FBoxAudioSampleEntry : public FBoxSampleEntry
+{
+public:
+	FBoxAudioSampleEntry(uint32 InFormat) : FBoxSampleEntry(InFormat)
+	{ }
+	virtual ~FBoxAudioSampleEntry() = default;
+	void SetChannelCount(uint16 InChannelCount)
+	{ ChannelCount = InChannelCount; }
+	bool SetSampleRate(uint32 InSampleRate)
+	{ SampleRate = InSampleRate; return SampleRate >= 0x10000; }
+	bool BuildData(const FBoxBuilderData& InOutBuilderData) override;
+protected:
+	uint32 SampleRate = 0;
+	uint16 ChannelCount = 0;
+};
+
+
+class FBoxSTSD : public FBoxFull
+{
+public:
+	FBoxSTSD() : FBoxFull(BoxAtom('s','t','s','d'))
+	{ }
+	virtual ~FBoxSTSD() = default;
+	void SetVersion(uint8 InVersion)
+	{ Version = InVersion; }
+	bool BuildData(const FBoxBuilderData& InOutBuilderData) override;
+protected:
+};
+
+
+class FBoxSTTS : public FBoxFull
+{
+public:
+	FBoxSTTS() : FBoxFull(BoxAtom('s','t','t','s'))
+	{ }
+	virtual ~FBoxSTTS() = default;
+	void AddSampleDuration(uint32 InSampleDuration)
+	{
+		const uint32 Dur = InSampleDuration;
+		if (SampleDurations.IsEmpty() || SampleDurations[SampleDurations.Num() - 1].Delta != Dur)
+		{
+			FEntry e;
+			e.Count = 1;
+			e.Delta = Dur;
+			SampleDurations.Emplace(e);
+		}
+		else
+		{
+			++SampleDurations[SampleDurations.Num() - 1].Count;
+		}
+	}
+	bool BuildData(const FBoxBuilderData& InOutBuilderData) override;
+protected:
+	struct FEntry
+	{
+		uint32 Count;
+		uint32 Delta;
+	};
+	TChunkedArray<FEntry, 64> SampleDurations;
+};
+
+
+class FBoxSTSS : public FBoxFull
+{
+public:
+	FBoxSTSS() : FBoxFull(BoxAtom('s','t','s','s'))
+	{ }
+	virtual ~FBoxSTSS() = default;
+	void AddSampleIndex(uint32 InSampleNumber)
+	{ SyncSampleIndices.Emplace(InSampleNumber); }
+	bool BuildData(const FBoxBuilderData& InOutBuilderData) override;
+protected:
+	TChunkedArray<uint32, 240> SyncSampleIndices;
+};
+
+
+class FBoxCTTS : public FBoxFull
+{
+public:
+	FBoxCTTS() : FBoxFull(BoxAtom('c','t','t','s'))
+	{ }
+	virtual ~FBoxCTTS() = default;
+	void AddSampleCompositionOffset(int32 InCompositionTimeOffset)
+	{
+		if (CompositionTimeOffsets.IsEmpty() || CompositionTimeOffsets[CompositionTimeOffsets.Num() - 1].Delta != InCompositionTimeOffset)
+		{
+			FEntry e;
+			e.Count = 1;
+			e.Delta = InCompositionTimeOffset;
+			CompositionTimeOffsets.Emplace(e);
+		}
+		else
+		{
+			++CompositionTimeOffsets[CompositionTimeOffsets.Num() - 1].Count;
+		}
+		if (!bHasNegativeEntries && InCompositionTimeOffset < 0)
+		{
+			bHasNegativeEntries = true;
+		}
+	}
+	bool HaveCompositionOffsets() const
+	{ return !(CompositionTimeOffsets.IsEmpty() || (CompositionTimeOffsets.Num() == 1 && CompositionTimeOffsets[0].Delta == 0)); }
+	bool HasNegativeCompositionOffsets() const
+	{ return bHasNegativeEntries; }
+	bool BuildData(const FBoxBuilderData& InOutBuilderData) override;
+protected:
+	struct FEntry
+	{
+		uint32 Count;
+		int32 Delta;
+	};
+	TChunkedArray<FEntry, 64> CompositionTimeOffsets;
+	bool bHasNegativeEntries = false;
+};
+
+
+class FBoxSTSC : public FBoxFull
+{
+public:
+	FBoxSTSC() : FBoxFull(BoxAtom('s','t','s','c'))
+	{ }
+	virtual ~FBoxSTSC() = default;
+	void AddEntry(uint32 InFirstChunk, uint32 InSamplesPerChunk, uint32 InSampleDescriptionIndex)
+	{
+		if (SampleToChunkEntries.IsEmpty() ||
+			SampleToChunkEntries[SampleToChunkEntries.Num() - 1].samples_per_chunk != InSamplesPerChunk ||
+			SampleToChunkEntries[SampleToChunkEntries.Num() - 1].sample_description_index != InSampleDescriptionIndex)
+		{
+			FEntry e;
+			e.first_chunk = InFirstChunk;
+			e.samples_per_chunk = InSamplesPerChunk;
+			e.sample_description_index = InSampleDescriptionIndex;
+			SampleToChunkEntries.Emplace(e);
+		}
+	}
+	bool BuildData(const FBoxBuilderData& InOutBuilderData) override;
+protected:
+	struct FEntry
+	{
+		uint32 first_chunk = 0;
+		uint32 samples_per_chunk = 0;
+		uint32 sample_description_index = 0;
+	};
+	TChunkedArray<FEntry, 32> SampleToChunkEntries;
+};
+
+
+class FBoxSTSZ : public FBoxFull
+{
+public:
+	FBoxSTSZ() : FBoxFull(BoxAtom('s','t','s','z'))
+	{ }
+	virtual ~FBoxSTSZ() = default;
+	void SetConstantSampleSize(uint32 InConstantSampleSize)
+	{ ConstantSampleSize = InConstantSampleSize; }
+	void AddSampleSize(uint32 InSampleSize)
+	{ SampleSizes.Emplace(InSampleSize); LargestSampleSize = InSampleSize > LargestSampleSize ? InSampleSize : LargestSampleSize; }
+	uint32 GetNumberOfSamples() const
+	{ return SampleSizes.Num(); }
+	bool BuildData(const FBoxBuilderData& InOutBuilderData) override;
+protected:
+	TChunkedArray<uint32, 240> SampleSizes;
+	uint32 ConstantSampleSize = 0;
+	uint32 LargestSampleSize = 0;
+	bool bUseSmallSizeBox = false;
+};
+
+
+class FBoxSTCO : public FBoxFull
+{
+public:
+	FBoxSTCO() : FBoxFull(BoxAtom('s','t','c','o'))
+	{ }
+	virtual ~FBoxSTCO() = default;
+	void AddChunkOffset(int64 InChunkOffset)
+	{
+		ChunkOffsets.Emplace(InChunkOffset);
+		if (InChunkOffset >= 0x100000000LL)
+		{
+			bNeeds64Bits = true;
+		}
+	}
+	uint32 GetNumberOfChunkOffsets() const
+	{ return ChunkOffsets.Num(); }
+	bool BuildData(const FBoxBuilderData& InOutBuilderData) override;
+protected:
+	TChunkedArray<int64, 32> ChunkOffsets;
+	bool bNeeds64Bits = false;
+};
+
+
+class FBoxMVEX : public FBoxBase
+{
+public:
+	FBoxMVEX() : FBoxBase(BoxAtom('m','v','e','x'))
+	{ }
+	virtual ~FBoxMVEX() = default;
+protected:
+};
+
+
+class FBoxMEHD : public FBoxFull
+{
+public:
+	FBoxMEHD() : FBoxFull(BoxAtom('m','e','h','d'))
+	{ }
+	virtual ~FBoxMEHD() = default;
+	void SetFragmentDuration(uint64 InFragmentDuration)
+	{ FragmentDuration = InFragmentDuration; }
+	bool BuildData(const FBoxBuilderData& InOutBuilderData) override;
+	bool UpdateFragmentDuration(uint64 InFragmentDuration);
+	bool RewriteChanges(FRewriteBoxChangeDelegate InRewriteBoxChangeDelegate) override;
+protected:
+	uint64 FragmentDuration = 0;
+	int64 FragmentDurationPatchOffset = -1;
+	bool bWasBuilt = false;
+};
+
+
+class FBoxTREX : public FBoxFull
+{
+public:
+	FBoxTREX() : FBoxFull(BoxAtom('t','r','e','x'))
+	{ }
+	virtual ~FBoxTREX() = default;
+	void SetTrackID(uint32 InTrackID)
+	{ TrackID = InTrackID; }
+	void SetDefaultSampleDuration(uint32 InDefaultSampleDuration)
+	{ default_sample_duration = InDefaultSampleDuration; }
+	void SetDefaultSampleSize(uint32 InDefaultSampleSize)
+	{ default_sample_size = InDefaultSampleSize; }
+	void SetDefaultSampleFlags(uint32 InDefaultSampleFlags)
+	{ default_sample_flags = InDefaultSampleFlags; }
+	uint32 GetDefaultSampleFlags() const
+	{ return default_sample_flags; }
+	bool BuildData(const FBoxBuilderData& InOutBuilderData) override;
+protected:
+	uint32 TrackID = 0;
+	uint32 default_sample_description_index = 1;
+	uint32 default_sample_duration = 0;
+	uint32 default_sample_size = 0;
+	uint32 default_sample_flags = 0;
+};
+
+
+class FBoxMOOF : public FBoxBase
+{
+public:
+	FBoxMOOF() : FBoxBase(BoxAtom('m','o','o','f'))
+	{ }
+	virtual ~FBoxMOOF() = default;
+protected:
+};
+
+
+class FBoxMFHD : public FBoxFull
+{
+public:
+	FBoxMFHD() : FBoxFull(BoxAtom('m','f','h','d'))
+	{ }
+	virtual ~FBoxMFHD() = default;
+	void SetSequenceNumber(uint32 InSequenceNumber)
+	{ SequenceNumber = InSequenceNumber; }
+	bool BuildData(const FBoxBuilderData& InOutBuilderData) override;
+protected:
+	uint32 SequenceNumber = 0;
+};
+
+
+class FBoxTRAF : public FBoxBase
+{
+public:
+	FBoxTRAF() : FBoxBase(BoxAtom('t','r','a','f'))
+	{ }
+	virtual ~FBoxTRAF() = default;
+protected:
+};
+
+
+class FBoxTFHD : public FBoxFull
+{
+public:
+	FBoxTFHD() : FBoxFull(BoxAtom('t','f','h','d'))
+	{ }
+	virtual ~FBoxTFHD() = default;
+	void SetTrackID(uint32 InTrackID)
+	{ TrackID = InTrackID; }
+	void SetDefaultBaseIsMoof(bool bInDefaultBaseIsMoof)
+	{ bDefaultBaseIsMoof = bInDefaultBaseIsMoof; }
+	void SetBaseDataOffset(uint64 InBaseDataOffset)
+	{ base_data_offset = InBaseDataOffset; }
+	void SetSampleDescriptionIndex(uint32 InSampleDescriptionIndex)
+	{ sample_description_index = InSampleDescriptionIndex; }
+	void SetDefaultSampleDuration(uint32 InDefaultSampleDuration)
+	{ default_sample_duration = InDefaultSampleDuration; }
+	void SetDefaultSampleSize(uint32 InDefaultSampleSize)
+	{ default_sample_size = InDefaultSampleSize; }
+	void SetDefaultSampleFlags(uint32 InDefaultSampleFlags)
+	{ default_sample_flags = InDefaultSampleFlags; }
+	TOptional<uint32> GetDefaultSampleFlags() const
+	{ return default_sample_flags; }
+
+	bool BuildData(const FBoxBuilderData& InOutBuilderData) override;
+	bool UpdateBaseDataOffset(uint64 InBaseDataOffset);
+	bool RewriteChanges(FRewriteBoxChangeDelegate InRewriteBoxChangeDelegate) override;
+protected:
+	uint32 TrackID = 0;
+	bool bDefaultBaseIsMoof = false;
+	TOptional<uint64> base_data_offset;
+	TOptional<uint32> sample_description_index;
+	TOptional<uint32> default_sample_duration;
+	TOptional<uint32> default_sample_size;
+	TOptional<uint32> default_sample_flags;
+	int64 BaseDataPatchOffset = -1;
+	bool bWasBuilt = false;
+};
+
+
+class FBoxTFDT : public FBoxFull
+{
+public:
+	FBoxTFDT() : FBoxFull(BoxAtom('t','f','d','t'))
+	{ }
+	virtual ~FBoxTFDT() = default;
+
+	void SetBaseMediaDecodeTime(uint64 InBaseMediaDecodeTime)
+	{ BaseMediaDecodeTime = InBaseMediaDecodeTime; }
+
+	bool BuildData(const FBoxBuilderData& InOutBuilderData) override;
+protected:
+	uint64 BaseMediaDecodeTime = 0;
+};
+
+
+class FBoxTRUN : public FBoxFull
+{
+public:
+	FBoxTRUN() : FBoxFull(BoxAtom('t','r','u','n'))
+	{ }
+	virtual ~FBoxTRUN() = default;
+
+	void SetDataOffset(int32 InDataOffset)
+	{ DataOffset = InDataOffset; }
+	// Set default flags as they appear in the `trex` or `tfhd` box.
+	// The `trun` has no default flags and we use this value only as a convenience for testing.
+	void SetDefaultSampleFlagsForTesting(uint32 InDefaultSampleFlags)
+	{ DefaultSampleFlags = InDefaultSampleFlags; }
+	void SetFirstSampleFlags(uint32 InFirstSampleFlags)
+	{ first_sample_flags = InFirstSampleFlags; }
+	void AddSample(uint32 InSampleDuration, uint32 InSampleSize, uint32 InSampleFlags, int32 InSampleCompositionTimeOffset)
+	{
+		FEntry e;
+		e.sample_duration = InSampleDuration;
+		e.sample_size = InSampleSize;
+		e.sample_flags = InSampleFlags;
+		if ((e.sample_composition_time_offset = InSampleCompositionTimeOffset) != 0)
+		{
+			bRequiresCompositionTimeOffsets = true;
+		}
+		if (Samples.Num())
+		{
+			if (Samples[Samples.Num()-1].sample_size != e.sample_size)
+			{
+				bHasDifferentSizes = true;
+			}
+			if (Samples[Samples.Num()-1].sample_duration != e.sample_duration)
+			{
+				bHasDifferentDurations = true;
+			}
+			if (Samples[Samples.Num()-1].sample_flags != e.sample_flags)
+			{
+				// Do samples starting from the 2nd onward have different flags?
+				bFlagsAfterFirstAreDifferent = Samples.Num() > 1;
+			}
+		}
+		Samples.Emplace(MoveTemp(e));
+	}
+
+	uint32 GetFirstSampleFlagsOrDefault() const
+	{ return Samples.Num() ? Samples[0].sample_flags : DefaultSampleFlags; }
+	uint32 GetSubsequentSampleFlagsOrDefault() const
+	{ return Samples.Num() > 1 ? Samples[1].sample_flags : DefaultSampleFlags; }
+	bool AreFlagsAfterFirstSampleDifferent() const
+	{ return bFlagsAfterFirstAreDifferent; }
+
+	bool RequiresCompositionTimeOffsets() const
+	{ return bRequiresCompositionTimeOffsets; }
+	bool HasDifferentSizes() const
+	{ return bHasDifferentSizes; }
+	bool HasDifferentDurations() const
+	{ return bHasDifferentDurations; }
+	void DoNotWriteDurations()
+	{ bDoNotWriteDurations = true; }
+	void DoNotWriteSizes()
+	{ bDoNotWriteSizes = true; }
+	void DoNotWriteFlags()
+	{ bDoNotWriteFlags = true; }
+
+	bool BuildData(const FBoxBuilderData& InOutBuilderData) override;
+protected:
+	struct FEntry
+	{
+		uint32 sample_duration = 0;
+		uint32 sample_size = 0;
+		uint32 sample_flags = 0;
+		int32 sample_composition_time_offset = 0;
+	};
+	TChunkedArray<FEntry, 64> Samples;
+	TOptional<int32> DataOffset;
+	TOptional<uint32> first_sample_flags;
+	uint32 DefaultSampleFlags = 0;
+	bool bRequiresCompositionTimeOffsets = false;
+	bool bHasDifferentSizes = false;
+	bool bHasDifferentDurations = false;
+	bool bFlagsAfterFirstAreDifferent = false;
+	bool bDoNotWriteDurations = false;
+	bool bDoNotWriteSizes = false;
+	bool bDoNotWriteFlags = false;
+	int64 DataOffsetPatchOffset = -1;
+};
+
+
+class FBoxSUBS : public FBoxFull
+{
+public:
+	FBoxSUBS() : FBoxFull(BoxAtom('s','u','b','s'))
+	{ }
+	virtual ~FBoxSUBS() = default;
+
+	using FSubSampleInfo = IMP4RawMuxer::FTrackSample::FSubSampleInfo;
+
+	void AddSubsamplesForSample(uint32 InSampleNumber, const TArray<FSubSampleInfo>& InSubsampleInfo)
+	{
+		check(InSampleNumber);	// samples start counting at 1, there is no 0.
+		// No need to record samples with no subsample information.
+		if (InSubsampleInfo.Num())
+		{
+			SubSamples.Emplace(FSSI({InSampleNumber, InSubsampleInfo}));
+		}
+	}
+
+	bool BuildData(const FBoxBuilderData& InOutBuilderData) override;
+protected:
+	struct FSSI
+	{
+		uint32 SampleNumber = 0;
+		TArray<FSubSampleInfo> Subs;
+	};
+	TChunkedArray<FSSI, 32> SubSamples;
+};
+
+
+class FBoxEDTS : public FBoxBase
+{
+public:
+	FBoxEDTS() : FBoxBase(BoxAtom('e','d','t','s'))
+	{ }
+	virtual ~FBoxEDTS() = default;
+protected:
+};
+
+
+class FBoxELST : public FBoxFull
+{
+public:
+	FBoxELST() : FBoxFull(BoxAtom('e','l','s','t'))
+	{ }
+	virtual ~FBoxELST() = default;
+	void SetEditDuration(uint64 InEditDuration)
+	{ EditDuration = InEditDuration; }
+	void SetMediaTime(int64 InMediaTime)
+	{ MediaTime = InMediaTime; }
+	bool UpdateEditDuration(uint64 InEditDuration);
+	bool RewriteChanges(FRewriteBoxChangeDelegate InRewriteBoxChangeDelegate) override;
+	bool BuildData(const FBoxBuilderData& InOutBuilderData) override;
+protected:
+	uint64 EditDuration = 0;
+	int64 MediaTime = 0;
+	int64 EditDurationPatchOffset = -1;
+	bool bWasBuilt = false;
+};
+
+
+}
+
diff --git a/Engine/Plugins/Media/MP4Utilities/Source/MP4Muxer/Private/MP4MuxerModule.cpp b/Engine/Plugins/Media/MP4Utilities/Source/MP4Muxer/Private/MP4MuxerModule.cpp
new file mode 100644
--- /dev/null
+++ b/Engine/Plugins/Media/MP4Utilities/Source/MP4Muxer/Private/MP4MuxerModule.cpp
@@ -0,0 +1,32 @@
+// Copyright Epic Games, Inc. All Rights Reserved.
+
+#include "MP4MuxerModule.h"
+#include "Misc/App.h"
+#include "Modules/ModuleManager.h"
+
+#define LOCTEXT_NAMESPACE "MP4MuxerModule"
+
+DEFINE_LOG_CATEGORY(LogMP4Muxer);
+
+class FMP4MuxerModule : public IModuleInterface
+{
+public:
+	void StartupModule() override
+	{
+	}
+
+	void ShutdownModule() override
+	{
+	}
+
+	bool SupportsDynamicReloading() override
+	{
+		return false;
+	}
+
+private:
+};
+
+IMPLEMENT_MODULE(FMP4MuxerModule, MP4Muxer);
+
+#undef LOCTEXT_NAMESPACE
diff --git a/Engine/Plugins/Media/MP4Utilities/Source/MP4Muxer/Private/MP4MuxerModule.h b/Engine/Plugins/Media/MP4Utilities/Source/MP4Muxer/Private/MP4MuxerModule.h
new file mode 100644
--- /dev/null
+++ b/Engine/Plugins/Media/MP4Utilities/Source/MP4Muxer/Private/MP4MuxerModule.h
@@ -0,0 +1,8 @@
+// Copyright Epic Games, Inc. All Rights Reserved.
+
+#pragma once
+
+#include "CoreTypes.h"
+#include "Logging/LogMacros.h"
+
+DECLARE_LOG_CATEGORY_EXTERN(LogMP4Muxer, Log, All);
diff --git a/Engine/Plugins/Media/MP4Utilities/Source/MP4Muxer/Public/MP4Muxer.h b/Engine/Plugins/Media/MP4Utilities/Source/MP4Muxer/Public/MP4Muxer.h
new file mode 100644
--- /dev/null
+++ b/Engine/Plugins/Media/MP4Utilities/Source/MP4Muxer/Public/MP4Muxer.h
@@ -0,0 +1,334 @@
+// Copyright Epic Games, Inc. All Rights Reserved.
+
+#pragma once
+
+#include "CoreTypes.h"
+#include "Templates/SharedPointer.h"
+#include "Containers/Array.h"
+#include "Containers/ArrayView.h"
+#include "Containers/UnrealString.h"
+#include "Misc/Timespan.h"
+#include "Misc/FrameRate.h"
+#include "Misc/TVariant.h"
+
+
+#define UE_API MP4MUXER_API
+
+/**
+ * This module implements a basic mp4 file multiplexer/writer.
+ *
+ * The data to be written is handled as RAW data and is not checked for validity.
+ * Data is also expected to be framed already, that is, one data packet corresponds
+ * to one media sample.
+ *
+ * If the packet data requires internal sample sizes, eg. NAL unit start codes to
+ * be replaced with the NAL size for AVC/HEVC video, this needs to have been done
+ * already as well.
+ *
+ * Codec specific data, like configuration boxes, also need to be set up already
+ * in an mp4 box structure and provided per box, which will also just be
+ * written as raw data.
+ */
+class IMP4RawMuxer
+{
+public:
+	virtual ~IMP4RawMuxer() = default;
+
+	UE_API static TSharedRef<IMP4RawMuxer, ESPMode::ThreadSafe> Create();
+
+	// Wraps binary data in a box. Useful for creating the additional FTrackSpec::SampleEntryBoxes
+	UE_API static TArray<uint8> WrapDataInBox(uint32 InBoxAtom, const TConstArrayView<uint8> InBoxData);
+
+	// Wraps binary data in a full box. Useful for creating the additional FTrackSpec::SampleEntryBoxes
+	UE_API static TArray<uint8> WrapDataInBox(uint32 InBoxAtom, uint8 InBoxVersion, uint32 InBoxFlags, const TConstArrayView<uint8> InBoxData);
+
+	struct FTrackSpec
+	{
+		struct FVideo
+		{
+			uint32 DisplayWidth = 0;
+			uint32 DisplayHeight = 0;
+			FString CompressorName;
+		};
+		struct FAudio
+		{
+			uint32 SamplingRate = 0;
+			uint16 NumberOfChannels = 0;
+		};
+
+		/**
+		 * ISO/IEC 14496-12 Section 8.5.2.1 "format"
+		 */
+		uint32 SampleEntryFormat = 0;
+
+		/**
+		 * Set to `true` if every sample in this track is independently
+		 * decodable and is a sync sample (keyframe).
+		 * Set to `false` if samples depend on other samples.
+		 */
+		bool bIsAllKeyframes = true;
+
+		/**
+		 * The number of ticks per second in which this tracks sample
+		 * decode and presentation time as well as the sample duration
+		 * is measured in.
+		 * Must not be zero.
+		 */
+		uint32 Timescale = 0;
+
+		/**
+		 * If every sample has the same duration you can set this fixed
+		 * value here for convenience and then specify 0 for duration in
+		 * the FTrackSample structure.
+		 */
+		uint32 ConstantSampleDuration = 0;
+
+		/**
+		 * Track type specific values that need to be set.
+		 * The type this variant contains determines if the track contains
+		 * video or audio.
+		 * The values in these structures are written as-is and should
+		 * therefore be set correctly.
+		 */
+		TVariant<FVideo,FAudio> Properties;
+
+		/**
+		 * Additional pre-built boxes to be added to the sample type
+		 * (see `SampleEntryFormat` above) that is written to the `stsd` box.
+		 * This should contain codec specific data like the configuration
+		 * record and other (optional) boxes.
+		 */
+		TArray<TArray<uint8>> SampleEntryBoxes;
+
+		/**
+		 * If set, a subsample box `subs` will be added to the `stsd` or
+		 * `traf` box with subsample information provided for each sample
+		 * when adding it to the track.
+		 * This value contains the 24 bit flags value for the `subs` box.
+		 * When written the `subs` box will always be of version 1 to use
+		 * 32 bit subsample sizes.
+		 */
+		TOptional<uint32> SubSampleFlags;
+
+		/**
+		 * Various flags to set in the `tkhd` box.
+		 * See ISO/IEC 14496-12:2022 - 8.3.2 Track header box
+		 */
+		uint16 Layer = 0;
+		uint16 AlternateGroup = 0;
+		bool bEnabled = true;
+		bool bInMovie = true;
+
+		/**
+		 * Language tag to set in the `mvhd` box.
+		 * Must be a 3 letter ISO-639-2/T code.
+		 */
+		uint8 Language[3] { 'u','n','d' };
+	};
+
+
+	struct FConfiguration
+	{
+		enum class EMuxMode
+		{
+			Standard,
+			Fragmented,
+		};
+
+		/**
+		 * Filename of the file to write.
+		 */
+		FString OutputFilename;
+
+		/**
+		 * Whether to write a standard mp4 (`moov` box will be placed at the end) or
+		 * a fragmented mp4.
+		 */
+		EMuxMode MuxMode = EMuxMode::Standard;
+
+		/**
+		 * Duration for each chunk of sample data.
+		 * Samples of each track will be collected until this duration is reached
+		 * (or the track has no more samples to write) and then written to disk as
+		 * a group of samples per track.
+		 *
+		 * This value must not less or equal to zero.
+		 * A typical value is 500ms. The larger the value the more data the muxer
+		 * needs to collect and keep, requiring more memory.
+		 */
+		FTimespan InterleaveDuration;
+	};
+
+	/**
+	 * Configures the muxer.
+	 */
+	virtual bool Configure(const FConfiguration& InConfiguration) = 0;
+
+	/**
+	 * Adds a track to be written.
+	 * Returns a track index used during muxing or -1 if the track
+	 * cannot be added.
+	 */
+	virtual int32 AddTrack(const FTrackSpec& InTrackSpec) = 0;
+
+	/**
+	 * State value reported to the registered `FStatusDelegate`.
+	 */
+	enum class EStatus
+	{
+		Failed,
+		Finished
+	};
+	DECLARE_DELEGATE_OneParam(FStatusDelegate, EStatus /*Status*/);
+
+	/**
+	 * Defines an entry in the sample request array the registered
+	 * `FRequestTrackDataDelegate` provides indicating for which
+	 * track(s) another sample is needed before a new chunk
+	 * (a group of samples) can be written.
+	 */
+	struct FSampleRequest
+	{
+		int32 TrackIndex = -1;
+		uint32 SampleNumber = 0;
+	};
+	DECLARE_DELEGATE_OneParam(FRequestTrackDataDelegate, const TArray<FSampleRequest>& /*TracksWantingSamples*/);
+
+	/**
+	 * This method starts the muxing process.
+	 *
+	 * While there are samples to be muxed:
+	 *   The muxer invokes the `FRequestTrackDataDelegate`, providing an array of track indices
+	 *   (the value returned by `AddTrack`) and the sample number needed.
+	 *   The delegate is invoked from a muxer internal worker thread and is called at any time
+	 *   the muxer needs (or still needs) another sample for the particular track.
+	 *
+	 *   NOTE:
+	 *     Due to the asynchronous nature it will happen that the `FRequestTrackDataDelegate`
+	 *     asks for the same sample of a track more than once.
+	 *     For example, if you want to mux three tracks with the track indices 0, 1 and 2 the
+	 *     first invocation of the delegate will ask for sample number 0 of tracks 0, 1 and 2.
+	 *     As soon as you provide sample 0 of track 0 by calling `AddTrackSample()` the delegate
+	 *     will be invoked again, this time asking for track 0 sample 1 and *still* for track
+	 *     1 sample 0 and track 2 sample 0 because these samples have not been received yet.
+	 *
+	 *     As a result this means that even if you are currently in the process of gathering
+	 *     sample 0 of track 1 in another parallel worker thread of your own you must not
+	 *     take this new sample request as a new workload. You have to keep track which
+	 *     track samples are currently being worked on and ignore any sample request that
+	 *     pops up twice (or more).
+	 *
+	 *   You provide the requested sample by calling `AddTrackSample()`. If this is the last
+	 *   sample for this track you need to set the `bIsFinalSample` flag.
+	 *   If you do not know that this will be the last sample at that point you can provide
+	 *   an empty sample on the next call with that flag set.
+	 *
+	 *   If during muxing an error occurs the `FStatusDelegate` will be invoked with an
+	 *   `EStatus::Failed`. Muxing will stop and you can get a human readable error
+	 *   message by calling `GetLastError()`.
+	 *   When muxing is complete and the file has been successully written the `FStatusDelegate`
+	 *   will be invoked with an `EStatus::Finished`.
+	 */
+	virtual bool Start(FRequestTrackDataDelegate InSampleRequestDelegate, FStatusDelegate InStatusDelegate) = 0;
+
+
+	/**
+	 * This structure defines the sample to be added to a track
+	 * during the muxing process.
+	 * All members should be set correctly as no plausibility checks
+	 * are performed except for
+	 *  - `Duration` not being zero or less
+	 *  - `Data.Num()` being less than 2^32. A single sample cannot
+	 *     be larger than 4 GiB in an ISO/IEC 14496-12:2022 file.
+	 */
+	struct FTrackSample
+	{
+		struct FSubSampleInfo
+		{
+			uint32 codec_specific_parameters = 0;
+			uint32 subsample_size = 0;
+			uint8 subsample_priority = 0;
+			uint8 discardable = 0;
+		};
+		TArray<FSubSampleInfo> SubSamples;
+
+		/**
+		 * The sample data to be written.
+		 * Must already be in the final format as it is written verbatim.
+		 */
+		TConstArrayView64<uint8> Data;
+
+		/**
+		 * Presentation duration of this sample, given in the tracks timescale units.
+		 * Must not be zero unless a constant duration has been set for the track
+		 * which must not be zero.
+		 * This value is the most important value as it is written to the `stts`
+		 * box directly and their relative sum determines the decode time (DTS).
+		 */
+		uint32 Duration = 0;
+
+		/**
+		 * Decode time, given in the tracks timescale units.
+		 * This value is NOT written into the file directly.
+		 * An ISO/IEC 14496-12:2022 file only stores DTS differences between
+		 * samples, implicitly starting at zero.
+		 * The first DTS may be used to generate an `elst` box if not zero.
+		 * For keyframe-only tracks that do not need a time offset in an
+		 * edit list you can leave both DTS and PTS set to zero for every
+		 * sample and only specify the sample duration.
+		 */
+		int64 DTS = 0;
+
+		/**
+		 * Presentation (composition) time, given in the tracks timescale units.
+		 * This value is ignored for keyframe-only tracks.
+		 */
+		int64 PTS = 0;
+
+		/**
+		 * The sequential number of the sample.
+		 * This value is not written to the file. It is used only as a verification
+		 * that this is the sample the `FRequestTrackDataDelegate` has requested
+		 * and should be set to that value when adding the sample.
+		 */
+		uint32 SampleNumber = 0;
+
+		/**
+		 * Indicates whether or not this sample is a sync sample that is independently
+		 * decodable. If this is a keyframe-only track this value is implicitly handled
+		 * as `true` and does not need to be set.
+		 */
+		bool bIsKeyframe = false;
+
+		/**
+		 * Set to `true` when this is the final sample to be added to the track.
+		 * Set to `false` if more samples follow or if you cannot make that determination
+		 * at this point in time. You may provide an empty sample later that has this
+		 * flag set and provides no sample data in `Data`.
+		 */
+		bool bIsFinalSample = false;
+	};
+
+	/**
+	 * Adds a sample to be written into the specified track as a response to a sample
+	 * request made by the `FRequestTrackDataDelegate`.
+	 * Data provided by the ArrayView will be copied into the muxer and can be released
+	 * on the user side.
+	 */
+	virtual bool AddTrackSample(int32 InTrack, const FTrackSample& InTrackSample) = 0;
+
+	/**
+	 * Stops muxing and closes the output file.
+	 * This can be called at any time to abort the process and must be called
+	 * when muxing is complete (either with or without error).
+	 */
+	virtual bool StopAndClose() = 0;
+
+	/**
+	 * Returns a human readable error message if muxing failed.
+	 */
+	virtual FString GetLastError() = 0;
+
+};
+
+#undef UE_API
