MP4Utilities: Added track name to muxer

#rnx

#Synced-CL 49757097

diff --git a/Engine/Plugins/Media/MP4Utilities/Source/MP4Boxes/Private/MP4Track.cpp b/Engine/Plugins/Media/MP4Utilities/Source/MP4Boxes/Private/MP4Track.cpp
--- a/Engine/Plugins/Media/MP4Utilities/Source/MP4Boxes/Private/MP4Track.cpp
+++ b/Engine/Plugins/Media/MP4Utilities/Source/MP4Boxes/Private/MP4Track.cpp
@@ -348,7 +348,21 @@ namespace MP4Boxes
 			TSharedPtr<FMP4BoxBase, ESPMode::ThreadSafe> NameBox = UdtaBox->FindBoxRecursive<FMP4BoxBase>(MP4Utilities::MakeBoxAtom('n','a','m','e'), 0);
 			if (NameBox.IsValid() && NameBox->GetBoxData().Num())
 			{
-				CommonMetadata.Name = FString::ConstructFromPtrSize(reinterpret_cast<const ANSICHAR*>(NameBox->GetBoxData().GetData()), NameBox->GetBoxData().Num());
+				// Some tools write the `name` box as a simple box, others as a full box.
+				// The `name` box is not defined by ISO/IEC 14496-12 but by Apple QuickTime (see https://developer.apple.com/documentation/quicktime-file-format/name_atom)
+				// where it is defined as a full box, so technically all tools would have to write it as a full box.
+				TConstArrayView<uint8> bd(NameBox->GetBoxData());
+				// Check if the first byte, which would be the version when this is a full box, is small enough to actually be a version number
+				// and also small enough to not be a printable character. Check the second byte as well, which would be bits 16-23 of the flags
+				// which is assumed to be zero for a full `name` box.
+				if (bd.Num() >= 4 && bd[0] < 10U && bd[1] == 0)
+				{
+					CommonMetadata.Name = FString::ConstructFromPtrSize(reinterpret_cast<const UTF8CHAR*>(NameBox->GetBoxData().GetData()+4), NameBox->GetBoxData().Num()-4);
+				}
+				else
+				{
+					CommonMetadata.Name = FString::ConstructFromPtrSize(reinterpret_cast<const UTF8CHAR*>(NameBox->GetBoxData().GetData()), NameBox->GetBoxData().Num());
+				}
 			}
 		}
 		TSharedPtr<FMP4BoxHDLR, ESPMode::ThreadSafe> HdlrBox = MdiaBox->FindBoxRecursive<FMP4BoxHDLR>(MP4Utilities::MakeBoxAtom('h','d','l','r'), 0);
diff --git a/Engine/Plugins/Media/MP4Utilities/Source/MP4Muxer/Private/MP4Muxer.cpp b/Engine/Plugins/Media/MP4Utilities/Source/MP4Muxer/Private/MP4Muxer.cpp
--- a/Engine/Plugins/Media/MP4Utilities/Source/MP4Muxer/Private/MP4Muxer.cpp
+++ b/Engine/Plugins/Media/MP4Utilities/Source/MP4Muxer/Private/MP4Muxer.cpp
@@ -442,6 +442,18 @@ bool FMP4RawMuxer::WriteFileStart()
 			return false;
 		}
 
+		// Track name, if any.
+		if (!TrackSpecs[TrackIdx].Name.IsEmpty())
+		{
+			MP4MuxerBoxes::FBoxUDTA* UDTA = new MP4MuxerBoxes::FBoxUDTA;
+			// Write the name in a simple box. Most tools expect it that way even though it
+			// should be a full box according to Apple (the `name` box is defined by QuickTime)
+			MP4MuxerBoxes::FBoxNAMEbase* NAME = new MP4MuxerBoxes::FBoxNAMEbase;
+			NAME->SetName(TrackSpecs[TrackIdx].Name);
+			TRAK->AddChild(UDTA);
+			UDTA->AddChild(NAME);
+		}
+
 		// DINF/DREF/URL
 		MP4MuxerBoxes::FBoxDINF* DINF = new MP4MuxerBoxes::FBoxDINF;
 		MINF->AddChild(DINF);
diff --git a/Engine/Plugins/Media/MP4Utilities/Source/MP4Muxer/Private/MP4MuxerBoxes.cpp b/Engine/Plugins/Media/MP4Utilities/Source/MP4Muxer/Private/MP4MuxerBoxes.cpp
--- a/Engine/Plugins/Media/MP4Utilities/Source/MP4Muxer/Private/MP4MuxerBoxes.cpp
+++ b/Engine/Plugins/Media/MP4Utilities/Source/MP4Muxer/Private/MP4MuxerBoxes.cpp
@@ -2,6 +2,7 @@
 
 #include "MP4MuxerBoxes.h"
 #include "Math/NumericLimits.h"
+#include "Containers/StringConv.h"
 
 namespace MP4MuxerBoxes
 {
@@ -1216,4 +1217,37 @@ bool FBoxELST::RewriteChanges(FRewriteBoxChangeDelegate InRewriteBoxChangeDelega
 }
 
 
+bool FBoxNAMEbase::BuildData(const FBoxBuilderData& InOutBuilderData)
+{
+	// Convert the name into an UTF8 array
+	TArray<uint8> UTF8Data;
+	FTCHARToUTF8 cnv(*Name);
+	int32 Len = cnv.Length();
+	UTF8Data.AddUninitialized(Len);
+	FMemory::Memcpy(UTF8Data.GetData(), cnv.Get(), Len);
+
+	FBoxBase::BuildData(InOutBuilderData);
+	FBoxDataWriter dw(CompiledBoxData);
+	dw.PutAlignedArray(UTF8Data);
+	Size = dw.GetBytePosition();
+	return true;
+}
+
+bool FBoxNAME::BuildData(const FBoxBuilderData& InOutBuilderData)
+{
+	// Convert the name into an UTF8 array
+	TArray<uint8> UTF8Data;
+	FTCHARToUTF8 cnv(*Name);
+	int32 Len = cnv.Length();
+	UTF8Data.AddUninitialized(Len);
+	FMemory::Memcpy(UTF8Data.GetData(), cnv.Get(), Len);
+
+	FBoxFull::BuildData(InOutBuilderData);
+	FBoxDataWriter dw(CompiledBoxData);
+	dw.PutAlignedArray(UTF8Data);
+	Size = dw.GetBytePosition();
+	return true;
+}
+
+
 }
diff --git a/Engine/Plugins/Media/MP4Utilities/Source/MP4Muxer/Private/MP4MuxerBoxes.h b/Engine/Plugins/Media/MP4Utilities/Source/MP4Muxer/Private/MP4MuxerBoxes.h
--- a/Engine/Plugins/Media/MP4Utilities/Source/MP4Muxer/Private/MP4MuxerBoxes.h
+++ b/Engine/Plugins/Media/MP4Utilities/Source/MP4Muxer/Private/MP4MuxerBoxes.h
@@ -938,5 +938,47 @@ protected:
 };
 
 
+class FBoxUDTA : public FBoxBase
+{
+public:
+	FBoxUDTA() : FBoxBase(BoxAtom('u','d','t','a'))
+	{ }
+	virtual ~FBoxUDTA() = default;
+protected:
+};
+
+
+class FBoxNAMEbase : public FBoxBase
+{
+public:
+	FBoxNAMEbase() : FBoxBase(BoxAtom('n','a','m','e'))
+	{ }
+	virtual ~FBoxNAMEbase() = default;
+
+	void SetName(const FString& InName)
+	{ Name = InName; }
+
+	bool BuildData(const FBoxBuilderData& InOutBuilderData) override;
+protected:
+	FString Name;
+};
+
+class FBoxNAME : public FBoxFull
+{
+public:
+	FBoxNAME() : FBoxFull(BoxAtom('n','a','m','e'))
+	{ }
+	virtual ~FBoxNAME() = default;
+
+	void SetName(const FString& InName)
+	{ Name = InName; }
+
+	bool BuildData(const FBoxBuilderData& InOutBuilderData) override;
+protected:
+	FString Name;
+};
+
+
+
 }
 
diff --git a/Engine/Plugins/Media/MP4Utilities/Source/MP4Muxer/Public/MP4Muxer.h b/Engine/Plugins/Media/MP4Utilities/Source/MP4Muxer/Public/MP4Muxer.h
--- a/Engine/Plugins/Media/MP4Utilities/Source/MP4Muxer/Public/MP4Muxer.h
+++ b/Engine/Plugins/Media/MP4Utilities/Source/MP4Muxer/Public/MP4Muxer.h
@@ -124,6 +124,12 @@ public:
 		 * Must be a 3 letter ISO-639-2/T code.
 		 */
 		uint8 Language[3] { 'u','n','d' };
+
+		/**
+		 * If set this gets written as an `udta`/`name` box into the track to
+		 * give the track a human readable label.
+		 */
+		FString Name;
 	};
 
 
