MP4Utilities: added writing of timecode track and opaque data tracks

#rnx

#Synced-CL 49835607

diff --git a/Engine/Plugins/Media/ElectraCodecs/Source/ElectraDecoders/Private/Utils/ElectraCodecFormatUtils.cpp b/Engine/Plugins/Media/ElectraCodecs/Source/ElectraDecoders/Private/Utils/ElectraCodecFormatUtils.cpp
--- a/Engine/Plugins/Media/ElectraCodecs/Source/ElectraDecoders/Private/Utils/ElectraCodecFormatUtils.cpp
+++ b/Engine/Plugins/Media/ElectraCodecs/Source/ElectraDecoders/Private/Utils/ElectraCodecFormatUtils.cpp
@@ -1075,7 +1075,7 @@ static int32 DetermineCodecTypeFromDCR_hvcc(Electra::FCodecTypeFormat& InOut)
 			vid.Profile.Constraints = dcr.GetGeneralConstraintIndicatorFlags();
 			vid.Profile.CompatibilityFlags = dcr.GetGeneralProfileCompatibilityFlags();
 			InOut.CSD = dcr.GetCodecSpecificData();
-			int32 FourCC = InOut.CSD.IsEmpty() ? Make4CC('h','v','c','1') : Make4CC('h','e','v','1');
+			int32 FourCC = !InOut.CSD.IsEmpty() ? Make4CC('h','v','c','1') : Make4CC('h','e','v','1');
 			if (InOut.FourCC && InOut.FourCC != FourCC)
 			{
 				UE_LOG(LogElectraFormatUtils, Log, TEXT("FourCC was set as %08x but identifed as %08x from HEVCDecoderConfigurationRecord"), InOut.FourCC, FourCC);
diff --git a/Engine/Plugins/Media/ElectraPlayer/Source/ElectraProtron/Private/Player/ElectraProtronPlayerImpl.cpp b/Engine/Plugins/Media/ElectraPlayer/Source/ElectraProtron/Private/Player/ElectraProtronPlayerImpl.cpp
--- a/Engine/Plugins/Media/ElectraPlayer/Source/ElectraProtron/Private/Player/ElectraProtronPlayerImpl.cpp
+++ b/Engine/Plugins/Media/ElectraPlayer/Source/ElectraProtron/Private/Player/ElectraProtronPlayerImpl.cpp
@@ -490,6 +490,11 @@ void FElectraProtronPlayer::FImpl::InternalOpen(const FString& InFilename)
 				UE_LOG(LogElectraProtron, Warning, TEXT("No decoder to handle sample type \"%s\" of track #%u, ignoring this track."), *CodecInfo.RFC6381, TrackID);
 			}
 		}
+		else if (CodecInfo.Type == Electra::FCodecTypeFormat::EType::Subtitle)
+		{
+			// TODO: Check with the subtitle decoder factory if this format can be decoded.
+			UE_LOG(LogElectraProtron, Verbose, TEXT("Subtitles in track #%u are not handled right now, ignoring this track."), TrackID);
+		}
 	}
 
 	Duration = MovieDuration.GetAsTimespan();
@@ -607,8 +612,8 @@ void FElectraProtronPlayer::FImpl::InternalOpen(const FString& InFilename)
 		{
 			auto CTIMBox = UdtaBox->FindBoxRecursive<MP4Boxes::FMP4BoxBase>(MP4Utilities::MakeBoxAtom(0xa9U,'T','I','M'), 0);
 			auto CTSCBox = UdtaBox->FindBoxRecursive<MP4Boxes::FMP4BoxBase>(MP4Utilities::MakeBoxAtom(0xa9U,'T','S','C'), 0);
-			//auto CTSZBox = UdtaBox->FindBoxRecursive<MP4Boxes::FMP4BoxBase>(MP4Utilities::MakeBoxAtom(0xa9U,'T','S','Z'), 0);
-			if (CTIMBox.IsValid() && CTSCBox.IsValid())
+			auto CTSZBox = UdtaBox->FindBoxRecursive<MP4Boxes::FMP4BoxBase>(MP4Utilities::MakeBoxAtom(0xa9U,'T','S','Z'), 0);
+			if (CTIMBox.IsValid() && CTSCBox.IsValid() && CTSZBox.IsValid())
 			{
 				auto GetValue = [](const TSharedPtr<MP4Boxes::FMP4BoxBase, ESPMode::ThreadSafe>& InBox) -> FString
 				{
@@ -627,7 +632,7 @@ void FElectraProtronPlayer::FImpl::InternalOpen(const FString& InFilename)
 				FTrackInfo::FFirstSampleTimecode tc;
 				FFrameRate fr;
 				tc.Timecode = GetValue(CTIMBox);
-				tc.Framerate = GetValue(CTSCBox);
+				tc.Framerate = FString::Printf(TEXT("%s/%s"), *GetValue(CTSCBox), *GetValue(CTSZBox));
 				TOptional<FTimecode> ptc = FTimecode::ParseTimecode(tc.Timecode);
 				if (ptc.IsSet() && TryParseString(fr, *tc.Framerate))
 				{
@@ -886,6 +891,52 @@ void FElectraProtronPlayer::FImpl::GetTrackCodecInfo(Electra::FCodecTypeFormat&
 		}
 
 	}
+	else if (SampleType == MP4Boxes::FMP4BoxSampleEntry::ESampleType::Subtitles || SampleType == MP4Boxes::FMP4BoxSampleEntry::ESampleType::TimedText)
+	{
+		TArray<TSharedPtr<MP4Boxes::FMP4BoxSampleEntry, ESPMode::ThreadSafe>> Subtitles;
+		Stsd->GetSampleDescriptions(Subtitles);
+		if (Subtitles.Num() != 1)
+		{
+			if (Subtitles.IsEmpty())
+			{
+				UE_LOG(LogElectraProtron, Warning, TEXT("No subtitle sample entry found in `stsd` box, ignoring track #%u."), InTrackID);
+			}
+			else
+			{
+				UE_LOG(LogElectraProtron, Warning, TEXT("More than one subtitle sample entry found in `stsd` box, ignoring track #%u."), InTrackID);
+			}
+			return;
+		}
+
+		switch(Subtitles[0]->GetType())
+		{
+			case MP4Utilities::MakeBoxAtom('e','n','c','t'):
+			case MP4Utilities::MakeBoxAtom('e','n','c','u'):
+			{
+				bIsEncrypted = true;
+				break;
+			}
+			default:
+			{
+				OutCodecInfo.Type = Electra::FCodecTypeFormat::EType::Subtitle;
+				OutCodecInfo.FourCC = Subtitles[0]->GetType();
+				OutCodecInfo.Properties.Emplace<Electra::FCodecTypeFormat::FSubtitle>();
+				Electra::FCodecTypeFormat::FSubtitle& csub = OutCodecInfo.Properties.Get<Electra::FCodecTypeFormat::FSubtitle>();
+				for(auto& Ch : Subtitles[0]->GetChildren())
+				{
+					OutCodecInfo.ExtraBoxes.Emplace(Ch->GetType(), Ch->GetBoxData());
+				}
+			/*
+				if (!CodecFormatUtils->PrepareCodecTypeFormat(OutCodecInfo, IElectraCodecFormatUtilsModularFeature::FCodecTypeFormatParams()))
+				{
+					UE_LOG(LogElectraProtron, Warning, TEXT("Failed to prepare the codec format type, ignoring track #%u."), InTrackID);
+					return;
+				}
+			*/
+				break;
+			}
+		}
+	}
 	else if (SampleType == MP4Boxes::FMP4BoxSampleEntry::ESampleType::QTFFTimecode)
 	{
 		TArray<TSharedPtr<MP4Boxes::FMP4BoxQTFFTimecodeSampleEntry, ESPMode::ThreadSafe>> Timecodes;
diff --git a/Engine/Plugins/Media/MP4Utilities/Source/MP4Boxes/Private/MP4Boxes.cpp b/Engine/Plugins/Media/MP4Utilities/Source/MP4Boxes/Private/MP4Boxes.cpp
--- a/Engine/Plugins/Media/MP4Utilities/Source/MP4Boxes/Private/MP4Boxes.cpp
+++ b/Engine/Plugins/Media/MP4Utilities/Source/MP4Boxes/Private/MP4Boxes.cpp
@@ -59,6 +59,7 @@ namespace MP4Boxes
 /****************************************************************************************************************************************************/
 	struct FMP4BoxTKHD::FParsed
 	{
+		TArray<int32> Matrix;
 		int64 Duration = 0;
 		uint32 Flags = 0;
 		uint32 TrackID = 0;
@@ -66,6 +67,7 @@ namespace MP4Boxes
 		uint32 Height = 0;
 		int16 Layer = 0;
 		int16 AlternateGroup = 0;
+		int16 Volume = 0;
 		uint8 Version = 0;
 	};
 	void FMP4BoxTKHD::ParseIfRequired()
@@ -101,7 +103,16 @@ namespace MP4Boxes
 		Parsed->Layer = (int16) Value16;
 		ar.Read(Value16);
 		Parsed->AlternateGroup = (int16) Value16;
-		ar.SkipBytes(4 + 9*4);				// `volume`, `reserved` and `matrix`
+		ar.Read(Value16);
+		Parsed->Volume = (int16) Value16;
+		ar.SkipBytes(2);					// `reserved`
+		Parsed->Matrix.SetNumUninitialized(9);
+		for(int32 i=0; i<9; ++i)
+		{
+			uint32 Value32;
+			ar.Read(Value32);
+			Parsed->Matrix[i] = (int32) Value32;
+		}
 		ar.Read(Parsed->Width);
 		ar.Read(Parsed->Height);
 	}
@@ -136,6 +147,28 @@ namespace MP4Boxes
 		return Parsed->Height >> 16;
 	}
 
+	TArray<int32> FMP4BoxTKHD::GetMatrix()
+	{
+		ParseIfRequired();
+		return Parsed->Matrix;
+	}
+	uint32 FMP4BoxTKHD::GetWidthRaw()
+	{
+		ParseIfRequired();
+		return Parsed->Width;
+	}
+	uint32 FMP4BoxTKHD::GetHeightRaw()
+	{
+		ParseIfRequired();
+		return Parsed->Height;
+	}
+	uint16 FMP4BoxTKHD::GetVolumeRaw()
+	{
+		ParseIfRequired();
+		return Parsed->Volume;
+	}
+
+
 /****************************************************************************************************************************************************/
 
 	struct FMP4BoxTREF::FParsed
@@ -420,6 +453,11 @@ namespace MP4Boxes
 				Parsed->SampleType = FMP4BoxSampleEntry::ESampleType::TimedMetadata;
 				break;
 			}
+			case MP4Utilities::MakeBoxAtom('t','e','x','t'):
+			{
+				Parsed->SampleType = FMP4BoxSampleEntry::ESampleType::TimedText;
+				break;
+			}
 			case MP4Utilities::MakeBoxAtom('s','b','t','l'):
 			case MP4Utilities::MakeBoxAtom('s','u','b','t'):
 			{
@@ -467,6 +505,7 @@ namespace MP4Boxes
 					break;
 				}
 				case FMP4BoxSampleEntry::ESampleType::Subtitles:
+				case FMP4BoxSampleEntry::ESampleType::TimedText:
 				{
 					switch(bi.Type)
 					{
@@ -475,6 +514,16 @@ namespace MP4Boxes
 							SampleEntry = FMP4BoxTX3GSampleEntry::Create(AsWeak(), bi);
 							break;
 						}
+						case MP4Utilities::MakeBoxAtom('s','t','p','p'):
+						{
+							SampleEntry = FMP4BoxXMLSubtitleSampleEntry::Create(AsWeak(), bi);
+							break;
+						}
+						case MP4Utilities::MakeBoxAtom('w','v','t','t'):
+						{
+							SampleEntry = FMP4BoxWVTTSampleEntry::Create(AsWeak(), bi);
+							break;
+						}
 						default:
 						{
 							//check(!"TODO");
@@ -570,7 +619,8 @@ namespace MP4Boxes
 
 	struct FMP4BoxVisualSampleEntry::FParsed
 	{
-		uint8 CompressorName[32] {};
+		FString CompressorName;
+		uint8 CompressorNameData[32] {};
 		uint16 DataReferenceIndex = 0;
 		uint16 Width = 0;
 		uint16 Height = 0;
@@ -594,9 +644,13 @@ namespace MP4Boxes
 		ar.Read(Parsed->Height);
 		ar.SkipBytes(12);		// `horizresolution`, `vertresolution`, `reserved`
 		ar.Read(Parsed->FrameCount);
-		ar.ReadBytes(Parsed->CompressorName, 32);
+		ar.ReadBytes(Parsed->CompressorNameData, 32);
 		ar.Read(Parsed->Depth);
 		ar.SkipBytes(2);		// `pre_defined`
+		if (Parsed->CompressorNameData[0] && Parsed->CompressorNameData[0] < 32)
+		{
+			Parsed->CompressorName = FString::ConstructFromPtrSize(reinterpret_cast<const UTF8CHAR*>(Parsed->CompressorNameData + 1), Parsed->CompressorNameData[0]);
+		}
 		// There can now be additional boxes following, most notably `pasp` and `clap`.
 		ProcessBoxChildrenRecursively(ar, BoxInfo);
 	}
@@ -626,6 +680,12 @@ namespace MP4Boxes
 		ParseIfRequired();
 		return Parsed->Depth;
 	}
+	FString FMP4BoxVisualSampleEntry::GetCompressorName()
+	{
+		ParseIfRequired();
+		return Parsed->CompressorName;
+	}
+
 
 /****************************************************************************************************************************************************/
 
@@ -774,6 +834,115 @@ namespace MP4Boxes
 		return Parsed->ConstLPCMFramesPerAudioPacket;
 	}
 
+/****************************************************************************************************************************************************/
+
+	struct FMP4BoxXMLSubtitleSampleEntry::FParsed
+	{
+		ESampleType SampleType = ESampleType::Subtitles;
+		uint16 DataReferenceIndex = 0;
+	};
+	void FMP4BoxXMLSubtitleSampleEntry::ParseIfRequired()
+	{
+		if (Parsed.IsValid())
+		{
+			return;
+		}
+		Parsed = MakePimpl<FParsed>();
+		MP4Utilities::FMP4AtomReader ar(BoxInfo.Data);
+		// Members of the generic `SampleEntry` class
+		ar.SkipBytes(6);
+		ar.Read(Parsed->DataReferenceIndex);
+/*
+	typedefs:
+		utf8list: null-terminated list of space-separated UTF-8 strings
+		utf8string: UTF-8 string as defined in IETF RFC 3629, null-terminated.
+	box members:
+		utf8list namespace;
+		utf8list schema_location; // optional
+		utf8list auxiliary_mime_types;
+		// optional, required if auxiliary resources are present
+
+		// TX3GSampleEntry follows
+		ar.SkipBytes(4);		// `reserved`
+		ar.Read(Parsed->Flags);
+		ar.Read(Parsed->Timescale);
+		ar.Read(Parsed->FrameDuration);
+		ar.Read(Parsed->NumberOfFrames);
+		ar.SkipBytes(1);		// `reserved`
+		// There can now be additional boxes following.
+		ProcessBoxChildrenRecursively(ar, BoxInfo);
+*/
+
+		uint32 HandlerType = 0;
+		auto MDIABox = FindParentBox<FMP4BoxMDIA>(MP4Utilities::MakeBoxAtom('m','d','i','a'));
+		if (MDIABox.IsValid())
+		{
+			auto HDLRBox = MDIABox->FindBoxRecursive<FMP4BoxHDLR>(MP4Utilities::MakeBoxAtom('h','d','l','r'), 1);
+			HandlerType = HDLRBox.IsValid() ? HDLRBox->GetHandlerType() : 0;
+		}
+		Parsed->SampleType = HandlerType == MP4Utilities::MakeBoxAtom('t','e','x','t') ? ESampleType::TimedText : ESampleType::Subtitles;
+	}
+
+	FMP4BoxSampleEntry::ESampleType FMP4BoxXMLSubtitleSampleEntry::GetSampleType()
+	{
+		ParseIfRequired();
+		return Parsed->SampleType;
+	}
+
+	uint16 FMP4BoxXMLSubtitleSampleEntry::GetDataReferenceIndex()
+	{
+		ParseIfRequired();
+		return Parsed->DataReferenceIndex;
+	}
+
+/****************************************************************************************************************************************************/
+
+	struct FMP4BoxWVTTSampleEntry::FParsed
+	{
+		ESampleType SampleType = ESampleType::Subtitles;
+		uint16 DataReferenceIndex = 0;
+	};
+	void FMP4BoxWVTTSampleEntry::ParseIfRequired()
+	{
+		if (Parsed.IsValid())
+		{
+			return;
+		}
+		Parsed = MakePimpl<FParsed>();
+		MP4Utilities::FMP4AtomReader ar(BoxInfo.Data);
+		// Members of the generic `SampleEntry` class
+		ar.SkipBytes(6);
+		ar.Read(Parsed->DataReferenceIndex);
+		/*
+			These boxes follow
+				WebVTTConfigurationBox config;
+				WebVTTSourceLabelBox label; // recommended
+				MPEG4BitRateBox (); // optional
+		*/
+		ProcessBoxChildrenRecursively(ar, BoxInfo);
+
+		uint32 HandlerType = 0;
+		auto MDIABox = FindParentBox<FMP4BoxMDIA>(MP4Utilities::MakeBoxAtom('m','d','i','a'));
+		if (MDIABox.IsValid())
+		{
+			auto HDLRBox = MDIABox->FindBoxRecursive<FMP4BoxHDLR>(MP4Utilities::MakeBoxAtom('h','d','l','r'), 1);
+			HandlerType = HDLRBox.IsValid() ? HDLRBox->GetHandlerType() : 0;
+		}
+		Parsed->SampleType = HandlerType == MP4Utilities::MakeBoxAtom('t','e','x','t') ? ESampleType::TimedText : ESampleType::Subtitles;
+	}
+
+	FMP4BoxSampleEntry::ESampleType FMP4BoxWVTTSampleEntry::GetSampleType()
+	{
+		ParseIfRequired();
+		return Parsed->SampleType;
+	}
+
+	uint16 FMP4BoxWVTTSampleEntry::GetDataReferenceIndex()
+	{
+		ParseIfRequired();
+		return Parsed->DataReferenceIndex;
+	}
+
 /****************************************************************************************************************************************************/
 
 	struct FMP4BoxQTFFTimecodeSampleEntry::FParsed
@@ -836,6 +1005,7 @@ namespace MP4Boxes
 
 	struct FMP4BoxTX3GSampleEntry::FParsed
 	{
+		ESampleType SampleType = ESampleType::Subtitles;
 		uint16 DataReferenceIndex = 0;
 	};
 	void FMP4BoxTX3GSampleEntry::ParseIfRequired()
@@ -860,6 +1030,20 @@ namespace MP4Boxes
 		// There can now be additional boxes following.
 		ProcessBoxChildrenRecursively(ar, BoxInfo);
 */
+		uint32 HandlerType = 0;
+		auto MDIABox = FindParentBox<FMP4BoxMDIA>(MP4Utilities::MakeBoxAtom('m','d','i','a'));
+		if (MDIABox.IsValid())
+		{
+			auto HDLRBox = MDIABox->FindBoxRecursive<FMP4BoxHDLR>(MP4Utilities::MakeBoxAtom('h','d','l','r'), 1);
+			HandlerType = HDLRBox.IsValid() ? HDLRBox->GetHandlerType() : 0;
+		}
+		Parsed->SampleType = HandlerType == MP4Utilities::MakeBoxAtom('t','e','x','t') ? ESampleType::TimedText : ESampleType::Subtitles;
+	}
+
+	FMP4BoxSampleEntry::ESampleType FMP4BoxTX3GSampleEntry::GetSampleType()
+	{
+		ParseIfRequired();
+		return Parsed->SampleType;
 	}
 
 	uint16 FMP4BoxTX3GSampleEntry::GetDataReferenceIndex()
@@ -963,6 +1147,29 @@ namespace MP4Boxes
 		return Parsed->AV1CodecConfigurationRecord;
 	}
 
+/****************************************************************************************************************************************************/
+
+	struct FMP4BoxAPVC::FParsed
+	{
+		TConstArrayView<uint8> APVCodecConfigurationRecord;
+	};
+	void FMP4BoxAPVC::ParseIfRequired()
+	{
+		if (Parsed.IsValid())
+		{
+			return;
+		}
+		Parsed = MakePimpl<FParsed>();
+		MP4Utilities::FMP4AtomReader ar(BoxInfo.Data);
+		Parsed->APVCodecConfigurationRecord = MakeConstArrayView(ar.GetCurrentDataPointer(), ar.GetNumBytesRemaining());
+	}
+
+	TConstArrayView<uint8> FMP4BoxAPVC::GetAVPCCodecConfigurationRecord()
+	{
+		ParseIfRequired();
+		return Parsed->APVCodecConfigurationRecord;
+	}
+
 /****************************************************************************************************************************************************/
 
 	struct FMP4BoxDVCC::FParsed
@@ -2700,6 +2907,7 @@ namespace MP4Boxes
 		FactoryMap.Add(MP4Utilities::MakeBoxAtom('v','m','h','d'), &FMP4BoxVMHD::Create);
 		FactoryMap.Add(MP4Utilities::MakeBoxAtom('s','m','h','d'), &FMP4BoxSMHD::Create);
 		FactoryMap.Add(MP4Utilities::MakeBoxAtom('n','m','h','d'), &FMP4BoxNMHD::Create);
+		FactoryMap.Add(MP4Utilities::MakeBoxAtom('s','t','h','d'), &FMP4BoxSTHD::Create);
 
 		FactoryMap.Add(MP4Utilities::MakeBoxAtom('g','m','h','d'), &FMP4BoxGMHD::Create);
 		FactoryMap.Add(MP4Utilities::MakeBoxAtom('t','a','p','t'), &FMP4BoxTAPT::Create);
@@ -2741,6 +2949,7 @@ namespace MP4Boxes
 		FactoryMap.Add(MP4Utilities::MakeBoxAtom('d','O','p','s'), &FMP4BoxDOPS::Create);
 		FactoryMap.Add(MP4Utilities::MakeBoxAtom('v','p','c','C'), &FMP4BoxVPCC::Create);
 		FactoryMap.Add(MP4Utilities::MakeBoxAtom('a','v','1','C'), &FMP4BoxAV1C::Create);
+		FactoryMap.Add(MP4Utilities::MakeBoxAtom('a','p','v','C'), &FMP4BoxAPVC::Create);
 	}
 	TSharedPtr<FMP4BoxBase, ESPMode::ThreadSafe> FMP4BoxFactory::Create(const TWeakPtr<FMP4BoxBase, ESPMode::ThreadSafe>& InParent, const MP4Utilities::FMP4BoxInfo& InBoxInfo)
 	{
diff --git a/Engine/Plugins/Media/MP4Utilities/Source/MP4Boxes/Public/MP4Boxes.h b/Engine/Plugins/Media/MP4Utilities/Source/MP4Boxes/Public/MP4Boxes.h
--- a/Engine/Plugins/Media/MP4Utilities/Source/MP4Boxes/Public/MP4Boxes.h
+++ b/Engine/Plugins/Media/MP4Utilities/Source/MP4Boxes/Public/MP4Boxes.h
@@ -135,15 +135,20 @@ namespace MP4Boxes
 		{ return !!(GetFlags() & 8); }
 
 		/**
-			* Returns the duration of this track, measured in the timescale of the `mvhd` box.
-			* This is just the value as stored in the box. You need to apply the timescale yourself.
-			*/
+		 * Returns the duration of this track, measured in the timescale of the `mvhd` box.
+		 * This is just the value as stored in the box. You need to apply the timescale yourself.
+		 */
 		UE_API int64 GetDuration();
 
 		UE_API uint32 GetTrackID();
 
 		UE_API uint16 GetWidth();
 		UE_API uint16 GetHeight();
+
+		UE_API TArray<int32> GetMatrix();
+		UE_API uint32 GetWidthRaw();
+		UE_API uint32 GetHeightRaw();
+		UE_API uint16 GetVolumeRaw();
 	protected:
 		FMP4BoxTKHD(const TWeakPtr<FMP4BoxBase, ESPMode::ThreadSafe>& InParent, const MP4Utilities::FMP4BoxInfo& InBoxInfo)
 			: FMP4BoxBase(InParent, InBoxInfo)
@@ -416,6 +421,7 @@ namespace MP4Boxes
 			Subtitles,
 			QTFFTimecode,
 			TimedMetadata,
+			TimedText,
 			Unsupported
 		};
 		virtual ESampleType GetSampleType()
@@ -490,6 +496,7 @@ namespace MP4Boxes
 		UE_API uint16 GetHeight();
 		UE_API uint16 GetFrameCount();
 		UE_API uint16 GetDepth();
+		UE_API FString GetCompressorName();
 	protected:
 		FMP4BoxVisualSampleEntry(const TWeakPtr<FMP4BoxBase, ESPMode::ThreadSafe>& InParent, const MP4Utilities::FMP4BoxInfo& InBoxInfo)
 			: FMP4BoxSampleEntry(InParent, InBoxInfo)
@@ -530,6 +537,50 @@ namespace MP4Boxes
 	};
 
 
+	/**
+	 * XML subtitle sample entry
+	 * ISO/IEC 14496-12:2022 - 12.6.3 Sample entry
+	 */
+	class FMP4BoxXMLSubtitleSampleEntry : public FMP4BoxSampleEntry
+	{
+	public:
+		static TSharedPtr<FMP4BoxBase, ESPMode::ThreadSafe> Create(const TWeakPtr<FMP4BoxBase, ESPMode::ThreadSafe>& InParent, const MP4Utilities::FMP4BoxInfo& InBoxInfo)
+		{ return MakeShareable<FMP4BoxXMLSubtitleSampleEntry>(new FMP4BoxXMLSubtitleSampleEntry(InParent, InBoxInfo)); }
+		virtual ~FMP4BoxXMLSubtitleSampleEntry() = default;
+		ESampleType GetSampleType() override;
+		uint16 GetDataReferenceIndex() override;
+	protected:
+		FMP4BoxXMLSubtitleSampleEntry(const TWeakPtr<FMP4BoxBase, ESPMode::ThreadSafe>& InParent, const MP4Utilities::FMP4BoxInfo& InBoxInfo)
+			: FMP4BoxSampleEntry(InParent, InBoxInfo)
+		{ }
+		void ParseIfRequired();
+		struct FParsed;
+		TPimplPtr<FParsed> Parsed;
+	};
+
+
+	/**
+	 * WebVTT sample entry
+	 * ISO/IEC 14496-30:2018 - 6.5 Sample entry format
+	 */
+	class FMP4BoxWVTTSampleEntry : public FMP4BoxSampleEntry
+	{
+	public:
+		static TSharedPtr<FMP4BoxBase, ESPMode::ThreadSafe> Create(const TWeakPtr<FMP4BoxBase, ESPMode::ThreadSafe>& InParent, const MP4Utilities::FMP4BoxInfo& InBoxInfo)
+		{ return MakeShareable<FMP4BoxWVTTSampleEntry>(new FMP4BoxWVTTSampleEntry(InParent, InBoxInfo)); }
+		virtual ~FMP4BoxWVTTSampleEntry() = default;
+		ESampleType GetSampleType() override;
+		uint16 GetDataReferenceIndex() override;
+	protected:
+		FMP4BoxWVTTSampleEntry(const TWeakPtr<FMP4BoxBase, ESPMode::ThreadSafe>& InParent, const MP4Utilities::FMP4BoxInfo& InBoxInfo)
+			: FMP4BoxSampleEntry(InParent, InBoxInfo)
+		{ }
+		void ParseIfRequired();
+		struct FParsed;
+		TPimplPtr<FParsed> Parsed;
+	};
+
+
 	/**
 	 * Timecode sample entry
 	 * Apple QuickTime: https://developer.apple.com/documentation/quicktime-file-format/timecode_sample_description
@@ -575,8 +626,7 @@ namespace MP4Boxes
 		static TSharedPtr<FMP4BoxBase, ESPMode::ThreadSafe> Create(const TWeakPtr<FMP4BoxBase, ESPMode::ThreadSafe>& InParent, const MP4Utilities::FMP4BoxInfo& InBoxInfo)
 		{ return MakeShareable<FMP4BoxTX3GSampleEntry>(new FMP4BoxTX3GSampleEntry(InParent, InBoxInfo)); }
 		virtual ~FMP4BoxTX3GSampleEntry() = default;
-		ESampleType GetSampleType() override
-		{ return ESampleType::Subtitles; }
+		ESampleType GetSampleType() override;
 		uint16 GetDataReferenceIndex() override;
 	protected:
 		FMP4BoxTX3GSampleEntry(const TWeakPtr<FMP4BoxBase, ESPMode::ThreadSafe>& InParent, const MP4Utilities::FMP4BoxInfo& InBoxInfo)
@@ -1096,9 +1146,9 @@ namespace MP4Boxes
 		{ return MakeShareable<FMP4BoxMEHD>(new FMP4BoxMEHD(InParent, InBoxInfo)); }
 		virtual ~FMP4BoxMEHD() = default;
 		/**
-			* Returns the fragment duration in the timescale of the `mvhd` box.
-			* This is just the value as stored in the box. You need to apply the timescale yourself.
-			*/
+		 * Returns the fragment duration in the timescale of the `mvhd` box.
+		 * This is just the value as stored in the box. You need to apply the timescale yourself.
+		 */
 		UE_API int64 GetFragmentDuration();
 	protected:
 		FMP4BoxMEHD(const TWeakPtr<FMP4BoxBase, ESPMode::ThreadSafe>& InParent, const MP4Utilities::FMP4BoxInfo& InBoxInfo)
@@ -1436,6 +1486,23 @@ namespace MP4Boxes
 	};
 
 
+	/**
+	 * `sthd` box
+	 * ISO/IEC 14496-12:2022 - 12.6.2 Subtitle media header
+	 */
+	class FMP4BoxSTHD : public FMP4BoxBase
+	{
+	public:
+		static TSharedPtr<FMP4BoxBase, ESPMode::ThreadSafe> Create(const TWeakPtr<FMP4BoxBase, ESPMode::ThreadSafe>& InParent, const MP4Utilities::FMP4BoxInfo& InBoxInfo)
+		{ return MakeShareable<FMP4BoxSTHD>(new FMP4BoxSTHD(InParent, InBoxInfo)); }
+		virtual ~FMP4BoxSTHD() = default;
+	protected:
+		FMP4BoxSTHD(const TWeakPtr<FMP4BoxBase, ESPMode::ThreadSafe>& InParent, const MP4Utilities::FMP4BoxInfo& InBoxInfo)
+			: FMP4BoxBase(InParent, InBoxInfo)
+		{ }
+	};
+
+
 	/**
 	 * `nmhd` box
 	 * ISO/IEC 14496-12:2022 - 8.4.5.2 Null media header box
@@ -1749,6 +1816,28 @@ namespace MP4Boxes
 	};
 
 
+	/**
+	 * `apvC` box
+	 * APV Codec ISO Media File Format Binding
+	 * See: [REDACTED_PROFILE_URL]/openapv/blob/main/readme/apv_isobmff.md#apv-codec-configuration-box
+	 */
+	class FMP4BoxAPVC : public FMP4BoxBase
+	{
+	public:
+		static TSharedPtr<FMP4BoxBase, ESPMode::ThreadSafe> Create(const TWeakPtr<FMP4BoxBase, ESPMode::ThreadSafe>& InParent, const MP4Utilities::FMP4BoxInfo& InBoxInfo)
+		{ return MakeShareable<FMP4BoxAPVC>(new FMP4BoxAPVC(InParent, InBoxInfo)); }
+		virtual ~FMP4BoxAPVC() = default;
+		UE_API TConstArrayView<uint8> GetAVPCCodecConfigurationRecord();
+	protected:
+		FMP4BoxAPVC(const TWeakPtr<FMP4BoxBase, ESPMode::ThreadSafe>& InParent, const MP4Utilities::FMP4BoxInfo& InBoxInfo)
+			: FMP4BoxBase(InParent, InBoxInfo)
+		{ }
+		void ParseIfRequired();
+		struct FParsed;
+		TPimplPtr<FParsed> Parsed;
+	};
+
+
 	/**
 	 * `iods` box
 	 * ISO/IEC 14496-14:2020 - 6.2 Object Descriptor Box
diff --git a/Engine/Plugins/Media/MP4Utilities/Source/MP4Muxer/Private/MP4Muxer.cpp b/Engine/Plugins/Media/MP4Utilities/Source/MP4Muxer/Private/MP4Muxer.cpp
--- a/Engine/Plugins/Media/MP4Utilities/Source/MP4Muxer/Private/MP4Muxer.cpp
+++ b/Engine/Plugins/Media/MP4Utilities/Source/MP4Muxer/Private/MP4Muxer.cpp
@@ -30,12 +30,16 @@ public:
 
 	bool Configure(const FConfiguration& InConfiguration) override;
 	int32 AddTrack(const FTrackSpec& InTrackSpec) override;
+	bool AddTrackReference(int32 InTrackIndexToAddReferenceTo, int32 InTrackIndexBeingReferenced, uint32 InReferenceType) override;
 	bool Start(FRequestTrackDataDelegate InSampleRequestDelegate, FStatusDelegate InStatusDelegate) override;
 	bool AddTrackSample(int32 InTrack, const FTrackSample& InTrackSample) override;
 	bool StopAndClose() override;
 	FString GetLastError() override
 	{ return LastError; }
 private:
+	using FReferencedTrackIndexList = TArray<int32>;
+	using FTrackReferences = TMap<uint32 /*ReferenceType*/, FReferencedTrackIndexList>;
+
 	struct FTrackMuxBoxes
 	{
 		MP4MuxerBoxes::FBoxTREX* TREX = nullptr;
@@ -111,6 +115,7 @@ private:
 
 	FConfiguration Configuration;
 	TArray<FTrackSpec> TrackSpecs;
+	TMap<int32, FTrackReferences> TrackReferences;
 	FRequestTrackDataDelegate SampleRequestDelegate;
 	FStatusDelegate StatusDelegate;
 	FString LastError;
@@ -215,6 +220,29 @@ int32 FMP4RawMuxer::AddTrack(const FTrackSpec& InTrackSpec)
 	return TrackSpecs.Num() - 1;
 }
 
+bool FMP4RawMuxer::AddTrackReference(int32 InTrackIndexToAddReferenceTo, int32 InTrackIndexBeingReferenced, uint32 InReferenceType)
+{
+	if (InTrackIndexToAddReferenceTo < 0 || InTrackIndexToAddReferenceTo >= TrackSpecs.Num())
+	{
+		SetError(TEXT("Invalid track index to add a reference to"));
+		return false;
+	}
+	if (InTrackIndexBeingReferenced < 0 || InTrackIndexBeingReferenced >= TrackSpecs.Num())
+	{
+		SetError(TEXT("Invalid index of referenced track"));
+		return false;
+	}
+	FTrackReferences& IdxRef = TrackReferences.FindOrAdd(InTrackIndexToAddReferenceTo);
+	FReferencedTrackIndexList& RefList = IdxRef.FindOrAdd(InReferenceType);
+	if (RefList.Contains(InTrackIndexBeingReferenced))
+	{
+		SetError(TEXT("Track has already been referenced before"));
+		return false;
+	}
+	RefList.Emplace(InTrackIndexBeingReferenced);
+	return true;
+}
+
 
 
 
@@ -351,6 +379,8 @@ bool FMP4RawMuxer::WriteFileStart()
 	{
 		bool bIsVideo = false;
 		bool bIsAudio = false;
+		bool bIsTimecode = false;
+		bool bIsOpaque = false;
 
 		MP4MuxerBoxes::FBoxTRAK* TRAK = new MP4MuxerBoxes::FBoxTRAK;
 		MOOV->AddChild(TRAK);
@@ -377,10 +407,46 @@ bool FMP4RawMuxer::WriteFileStart()
 		}
 		else if (TrackSpecs[TrackIdx].Properties.IsType<FTrackSpec::FAudio>())
 		{
-			auto const& vid = TrackSpecs[TrackIdx].Properties.Get<FTrackSpec::FAudio>();
+			//auto const& aud = TrackSpecs[TrackIdx].Properties.Get<FTrackSpec::FAudio>();
 			TKHD->SetVolume(1);
 			bIsAudio = true;
 		}
+		else if (TrackSpecs[TrackIdx].Properties.IsType<FTrackSpec::FTimecode>())
+		{
+			bIsTimecode = true;
+		}
+		else if (TrackSpecs[TrackIdx].Properties.IsType<FTrackSpec::FOpaqueData>())
+		{
+			auto const& op = TrackSpecs[TrackIdx].Properties.Get<FTrackSpec::FOpaqueData>();
+			TKHD->SetMatrix(op.matrix);
+			TKHD->SetWidthRaw(op.width);
+			TKHD->SetHeightRaw(op.height);
+			TKHD->SetVolumeRaw(op.volume);
+			TKHD->SetCustomFlags(op.flags);
+			bIsOpaque = true;
+		}
+		else
+		{
+			unimplemented();
+		}
+
+		// Does this track reference anything?
+		const FTrackReferences* TrackRefs = TrackReferences.Find(TrackIdx);
+		if (TrackRefs)
+		{
+			MP4MuxerBoxes::FBoxTREF* TREF = new MP4MuxerBoxes::FBoxTREF;
+			TRAK->AddChild(TREF);
+			for(auto &TrRefIt : *TrackRefs)
+			{
+				MP4MuxerBoxes::FBoxTREFType* TREFType = new MP4MuxerBoxes::FBoxTREFType;
+				TREFType->SetReferenceType(TrRefIt.Key);
+				for(int32 i=0, iMax=TrRefIt.Value.Num(); i<iMax; ++i)
+				{
+					TREFType->AddReferencedTrackNumber(1 + TrRefIt.Value[i]);
+				}
+				TREF->AddChild(TREFType);
+			}
+		}
 
 		// TREX ?
 		if (MVEX)
@@ -436,6 +502,22 @@ bool FMP4RawMuxer::WriteFileStart()
 			MP4MuxerBoxes::FBoxSMHD* SMHD = new MP4MuxerBoxes::FBoxSMHD;
 			MINF->AddChild(SMHD);
 		}
+		else if (bIsTimecode)
+		{
+			HDLR->SetHandlerType(MP4MuxerBoxes::BoxAtom('t','m','c','d'));
+			HDLR->SetName(TEXT("TimeCodeHandler"));
+			MP4MuxerBoxes::FBoxNMHD* NMHD = new MP4MuxerBoxes::FBoxNMHD;
+			MINF->AddChild(NMHD);
+		}
+		else if (bIsOpaque)
+		{
+			auto const& op = TrackSpecs[TrackIdx].Properties.Get<FTrackSpec::FOpaqueData>();
+			HDLR->SetHandlerType(op.mdia_hdlr_Type);
+			HDLR->SetName(op.mdia_hdlr_Name);
+			MP4MuxerBoxes::FBoxOpaqueData* MediaHeader = new MP4MuxerBoxes::FBoxOpaqueData(op.minf_MediaHeader_Type);
+			MediaHeader->SetRawBoxData(op.minf_MediaHeader_RawBoxData);
+			MINF->AddChild(MediaHeader);
+		}
 		else
 		{
 			unimplemented();
@@ -489,6 +571,25 @@ bool FMP4RawMuxer::WriteFileStart()
 			Audio->SetAdditionalBoxes(TrackSpecs[TrackIdx].SampleEntryBoxes);
 			STSD->AddChild(Audio);
 		}
+		else if (bIsTimecode)
+		{
+			auto const& tc = TrackSpecs[TrackIdx].Properties.Get<FTrackSpec::FTimecode>();
+			MP4MuxerBoxes::FBoxTimecodeSampleEntry* Timecode = new MP4MuxerBoxes::FBoxTimecodeSampleEntry(TrackSpecs[TrackIdx].SampleEntryFormat);
+			Timecode->SetTimescale(tc.Timescale);
+			Timecode->SetFrameDuration(tc.FrameDuration);
+			Timecode->SetFramesPerSecond(tc.FramesPerSecond);
+			Timecode->SetDropFrame(tc.bDropFrame);
+			Timecode->SetMax24Hours(tc.bMax24Hours);
+			Timecode->SetAllowNegativeTimes(tc.bAllowNegativeTimes);
+			STSD->AddChild(Timecode);
+		}
+		else if (bIsOpaque)
+		{
+			auto const& od = TrackSpecs[TrackIdx].Properties.Get<FTrackSpec::FOpaqueData>();
+			MP4MuxerBoxes::FBoxOpaqueData* SampleEntry = new MP4MuxerBoxes::FBoxOpaqueData(TrackSpecs[TrackIdx].SampleEntryFormat);
+			SampleEntry->SetRawBoxData(od.stsd_SampleEntry_RawBoxData);
+			STSD->AddChild(SampleEntry);
+		}
 		else
 		{
 			unimplemented();
@@ -773,10 +874,16 @@ bool FMP4RawMuxer::AddTrackSample(int32 InTrack, const FTrackSample& InTrackSamp
 		uint32 Duration = InTrackSample.Duration ? InTrackSample.Duration : TrackSpecs[InTrack].ConstantSampleDuration;
 		const bool bIsKeyframeOnly = TrackSpecs[InTrack].bIsAllKeyframes;
 		bool bIsKeyframe = bIsKeyframeOnly ? true : InTrackSample.bIsKeyframe;
-		if (Duration == 0)
-		{
-			return SetError(FString::Printf(TEXT("Sample duration must not be zero or less")));
-		}
+		/*
+			Zero duration _may_ be set for the very last sample only. We don't really know if this is the last sample
+			(unless `InTrackSample.bIsFinalSample` is set, which is does not have to be) so we cannot reliably detect
+			this to be valid or not.
+
+			if (Duration == 0)
+			{
+				return SetError(FString::Printf(TEXT("Sample duration must not be zero or less")));
+			}
+		*/
 		// There is no version of a `stsz`, `stz2` or `trun` that would allow for a single AU to be larger than 4GiB (32 bits)
 		if (InTrackSample.Data.NumBytes() > 0xffffffffU)
 		{
diff --git a/Engine/Plugins/Media/MP4Utilities/Source/MP4Muxer/Private/MP4MuxerBoxes.cpp b/Engine/Plugins/Media/MP4Utilities/Source/MP4Muxer/Private/MP4MuxerBoxes.cpp
--- a/Engine/Plugins/Media/MP4Utilities/Source/MP4Muxer/Private/MP4MuxerBoxes.cpp
+++ b/Engine/Plugins/Media/MP4Utilities/Source/MP4Muxer/Private/MP4MuxerBoxes.cpp
@@ -409,6 +409,7 @@ bool FBoxTKHD::BuildData(const FBoxBuilderData& InOutBuilderData)
 	// Any of the values require a version 1 box?
 	Version = (CreationTime > 0xffffffffU || ModificationTime > 0xffffffffU || (Duration != ~((uint64)0) && Duration > 0xffffffffU)) ? 1 : 0;
 
+	Flags |= CustomFlags;
 	FBoxFull::BuildData(InOutBuilderData);
 	FBoxDataWriter dw(CompiledBoxData);
 	if (Version == 1)
@@ -434,16 +435,16 @@ bool FBoxTKHD::BuildData(const FBoxBuilderData& InOutBuilderData)
 	dw.PutAlignedValue(AlternateGroup);
 	dw.PutAlignedValue(Volume);
 	dw.PutAlignedValue((uint16)0);
-	// unity matrix
-	dw.PutAlignedValue((uint32)0x10000);
-	dw.PutAlignedValue((uint32)0);
-	dw.PutAlignedValue((uint32)0);
-	dw.PutAlignedValue((uint32)0);
-	dw.PutAlignedValue((uint32)0x10000);
-	dw.PutAlignedValue((uint32)0);
-	dw.PutAlignedValue((uint32)0);
-	dw.PutAlignedValue((uint32)0x10000);
-	dw.PutAlignedValue((uint32)0x40000000);
+	// matrix
+	dw.PutAlignedValue(Matrix[0]);
+	dw.PutAlignedValue(Matrix[1]);
+	dw.PutAlignedValue(Matrix[2]);
+	dw.PutAlignedValue(Matrix[3]);
+	dw.PutAlignedValue(Matrix[4]);
+	dw.PutAlignedValue(Matrix[5]);
+	dw.PutAlignedValue(Matrix[6]);
+	dw.PutAlignedValue(Matrix[7]);
+	dw.PutAlignedValue(Matrix[8]);
 	dw.PutAlignedValue(Width);
 	dw.PutAlignedValue(Height);
 	Size = dw.GetBytePosition();
@@ -707,6 +708,31 @@ bool FBoxAudioSampleEntry::BuildData(const FBoxBuilderData& InOutBuilderData)
 }
 
 
+bool FBoxTimecodeSampleEntry::BuildData(const FBoxBuilderData& InOutBuilderData)
+{
+	FBoxSampleEntry::BuildData(InOutBuilderData);
+	FBoxDataWriter dw(CompiledBoxData);
+	uint32 tcFlags = 0;
+	tcFlags |= bDropFrame ? 0x1U : 0x0U;
+	tcFlags |= bMax24Hours ? 0x2U : 0x0U;
+	tcFlags |= bAllowNegativeTimes ? 0x4U : 0x0U;
+	// Note: Do not set flag 0x8 to indicate counter values.
+	//  It is always counter values nowadays.
+	dw.PutAlignedValue((uint32)0);				// reserved
+	dw.PutAlignedValue(tcFlags);
+	dw.PutAlignedValue((uint32)Timescale);
+	dw.PutAlignedValue((uint32)FrameDuration);
+	dw.PutAlignedValue((uint8)FramesPerSecond);
+	dw.PutAlignedValue((uint8)0U);				// reserved
+	for(int32 i=0; i<AdditionalBoxes.Num(); ++i)
+	{
+		dw.PutAlignedArray(AdditionalBoxes[i]);
+	}
+	Size = dw.GetBytePosition();
+	return true;
+}
+
+
 bool FBoxSTSD::BuildData(const FBoxBuilderData& InOutBuilderData)
 {
 	FBoxFull::BuildData(InOutBuilderData);
@@ -836,10 +862,14 @@ bool FBoxSTSZ::BuildData(const FBoxBuilderData& InOutBuilderData)
 	{
 		FBoxFull::BuildData(InOutBuilderData);
 		FBoxDataWriter dw(CompiledBoxData);
+		if (ConstantSampleSize == 0 && !bHasVaryingSizes && SampleSizes.Num())
+		{
+			ConstantSampleSize = SampleSizes[0];
+		}
 		dw.PutAlignedValue(ConstantSampleSize);
+		dw.PutAlignedValue((uint32)SampleSizes.Num());
 		if (ConstantSampleSize == 0)
 		{
-			dw.PutAlignedValue((uint32)SampleSizes.Num());
 			for(int32 i=0,iMax=SampleSizes.Num(); i<iMax; ++i)
 			{
 				dw.PutAlignedValue(SampleSizes[i]);
@@ -1250,4 +1280,29 @@ bool FBoxNAME::BuildData(const FBoxBuilderData& InOutBuilderData)
 }
 
 
+bool FBoxTREFType::BuildData(const FBoxBuilderData& InOutBuilderData)
+{
+	FBoxBase::SetType(ReferenceType);
+	FBoxBase::BuildData(InOutBuilderData);
+	FBoxDataWriter dw(CompiledBoxData);
+	for(int32 i=0; i<ReferencedTrackNumbers.Num(); ++i)
+	{
+		dw.PutAlignedValue(ReferencedTrackNumbers[i]);
+	}
+	Size = dw.GetBytePosition();
+	return true;
+}
+
+
+bool FBoxOpaqueData::BuildData(const FBoxBuilderData& InOutBuilderData)
+{
+	//FBoxBase::SetType(ReferenceType);
+	FBoxBase::BuildData(InOutBuilderData);
+	FBoxDataWriter dw(CompiledBoxData);
+	dw.PutAlignedArray(RawBoxData);
+	Size = dw.GetBytePosition();
+	return true;
+}
+
+
 }
diff --git a/Engine/Plugins/Media/MP4Utilities/Source/MP4Muxer/Private/MP4MuxerBoxes.h b/Engine/Plugins/Media/MP4Utilities/Source/MP4Muxer/Private/MP4MuxerBoxes.h
--- a/Engine/Plugins/Media/MP4Utilities/Source/MP4Muxer/Private/MP4MuxerBoxes.h
+++ b/Engine/Plugins/Media/MP4Utilities/Source/MP4Muxer/Private/MP4MuxerBoxes.h
@@ -250,19 +250,31 @@ public:
 	{ Height = InHeight << 16; }
 	void SetVolume(uint8 InVolume)
 	{ Volume = (uint16)InVolume << 8; }
+	void SetWidthRaw(uint32 InWidth)
+	{ Width = InWidth; }
+	void SetHeightRaw(uint32 InHeight)
+	{ Height = InHeight; }
+	void SetVolumeRaw(int16 InVolume)
+	{ Volume = (uint16)InVolume; }
 	void SetLayer(uint16 InLayer)
 	{ Layer = InLayer; }
 	void SetAlternateGroup(uint16 InAlternateGroup)
 	{ AlternateGroup = InAlternateGroup; }
+	void SetMatrix(const int32 InMatrix[9])
+	{ FMemory::Memcpy(Matrix, InMatrix, sizeof(Matrix)); }
+	void SetCustomFlags(uint32 InFlags)
+	{ CustomFlags = InFlags & ~3U; }
 
 	bool BuildData(const FBoxBuilderData& InOutBuilderData) override;
 
 	bool UpdateDuration(uint64 InDuration);
 	bool RewriteChanges(FRewriteBoxChangeDelegate InRewriteBoxChangeDelegate) override;
 protected:
+	int32 Matrix[9] { 0x00010000,0,0,0,0x00010000,0,0,0,0x40000000 };
 	uint64 CreationTime = 0;
 	uint64 ModificationTime = 0;
 	uint64 Duration = 0;
+	uint32 CustomFlags = 0;
 	uint32 TrackID = 0;
 	uint32 Width = 0;
 	uint32 Height = 0;
@@ -377,6 +389,16 @@ protected:
 };
 
 
+class FBoxNMHD : public FBoxFull
+{
+public:
+	FBoxNMHD() : FBoxFull(BoxAtom('n','m','h','d'))
+	{ }
+	virtual ~FBoxNMHD() = default;
+protected:
+};
+
+
 class FBoxDINF : public FBoxBase
 {
 public:
@@ -472,6 +494,39 @@ protected:
 };
 
 
+class FBoxTimecodeSampleEntry : public FBoxSampleEntry
+{
+public:
+	FBoxTimecodeSampleEntry(uint32 InFormat) : FBoxSampleEntry(InFormat)
+	{ }
+	virtual ~FBoxTimecodeSampleEntry() = default;
+
+	void SetTimescale(int32 InTimescale)
+	{ Timescale = InTimescale; }
+	void SetFrameDuration(int32 InFrameDuration)
+	{ FrameDuration = InFrameDuration; }
+	void SetFramesPerSecond(int8 InFramesPerSecond)
+	{ FramesPerSecond = InFramesPerSecond; }
+	void SetDropFrame(bool bInDropFrame)
+	{ bDropFrame = bInDropFrame; }
+	void SetMax24Hours(bool bInMax24h)
+	{ bMax24Hours = bInMax24h; }
+	void SetAllowNegativeTimes(bool bInAllowNegativeTimes)
+	{ bAllowNegativeTimes = bInAllowNegativeTimes; }
+
+	bool BuildData(const FBoxBuilderData& InOutBuilderData) override;
+protected:
+	// These values are defined to be signed integers!
+	int32 Timescale = 0;
+	int32 FrameDuration = 0;
+	int8 FramesPerSecond = 0;
+
+	bool bDropFrame = false;
+	bool bMax24Hours = false;
+	bool bAllowNegativeTimes = false;
+};
+
+
 class FBoxSTSD : public FBoxFull
 {
 public:
@@ -611,7 +666,11 @@ public:
 	void SetConstantSampleSize(uint32 InConstantSampleSize)
 	{ ConstantSampleSize = InConstantSampleSize; }
 	void AddSampleSize(uint32 InSampleSize)
-	{ SampleSizes.Emplace(InSampleSize); LargestSampleSize = InSampleSize > LargestSampleSize ? InSampleSize : LargestSampleSize; }
+	{
+		bHasVaryingSizes = bHasVaryingSizes || (!SampleSizes.IsEmpty() && SampleSizes[SampleSizes.Num()-1] != InSampleSize);
+		SampleSizes.Emplace(InSampleSize);
+		LargestSampleSize = InSampleSize > LargestSampleSize ? InSampleSize : LargestSampleSize;
+	}
 	uint32 GetNumberOfSamples() const
 	{ return SampleSizes.Num(); }
 	bool BuildData(const FBoxBuilderData& InOutBuilderData) override;
@@ -619,6 +678,7 @@ protected:
 	TChunkedArray<uint32, 240> SampleSizes;
 	uint32 ConstantSampleSize = 0;
 	uint32 LargestSampleSize = 0;
+	bool bHasVaryingSizes = false;
 	bool bUseSmallSizeBox = false;
 };
 
@@ -979,6 +1039,46 @@ protected:
 };
 
 
+class FBoxTREF : public FBoxBase
+{
+public:
+	FBoxTREF() : FBoxBase(BoxAtom('t','r','e','f'))
+	{ }
+	virtual ~FBoxTREF() = default;
+protected:
+};
+
+class FBoxTREFType : public FBoxBase
+{
+public:
+	FBoxTREFType() : FBoxBase(0)
+	{ }
+	virtual ~FBoxTREFType() = default;
+
+	void SetReferenceType(uint32 InReferenceType)
+	{ ReferenceType = InReferenceType; }
+	void AddReferencedTrackNumber(uint32 InReferencedTrackNumber)
+	{ ReferencedTrackNumbers.Emplace(InReferencedTrackNumber); }
+	bool BuildData(const FBoxBuilderData& InOutBuilderData) override;
+protected:
+	uint32 ReferenceType = 0;
+	TArray<uint32> ReferencedTrackNumbers;
+};
+
+
+class FBoxOpaqueData : public FBoxBase
+{
+public:
+	FBoxOpaqueData(uint32 InType) : FBoxBase(InType)
+	{ }
+	virtual ~FBoxOpaqueData() = default;
+
+	void SetRawBoxData(TConstArrayView<uint8> InRawBoxData)
+	{ RawBoxData = InRawBoxData; }
+	bool BuildData(const FBoxBuilderData& InOutBuilderData) override;
+protected:
+	TArray<uint8> RawBoxData;
+};
 
 }
 
diff --git a/Engine/Plugins/Media/MP4Utilities/Source/MP4Muxer/Public/MP4Muxer.h b/Engine/Plugins/Media/MP4Utilities/Source/MP4Muxer/Public/MP4Muxer.h
--- a/Engine/Plugins/Media/MP4Utilities/Source/MP4Muxer/Public/MP4Muxer.h
+++ b/Engine/Plugins/Media/MP4Utilities/Source/MP4Muxer/Public/MP4Muxer.h
@@ -10,6 +10,7 @@
 #include "Misc/Timespan.h"
 #include "Misc/FrameRate.h"
 #include "Misc/TVariant.h"
+#include "Misc/Optional.h"
 
 
 #define UE_API MP4MUXER_API
@@ -55,6 +56,77 @@ public:
 			uint32 SamplingRate = 0;
 			uint16 NumberOfChannels = 0;
 		};
+		struct FTimecode
+		{
+			/*
+				Timecode is always written as `tmcd` so this are its parameters
+				See https://developer.apple.com/documentation/quicktime-file-format/timecode_sample_description
+				Please note that the "tape counter" flag is no longer mentioned in https://developer.apple.com/library/archive/technotes/tn2310/_index.html
+				and the technote seems to suggest that only frame counter values are written.
+				Also see https://developer.apple.com/documentation/coremedia/time-code-flags?changes=_5_3&language=objc
+				As such the deprecated "counter" flag (value was 0x8) will not be written.
+
+				You can write timecode tracks with as many samples as the associated video track or write
+				only single timecode sample that is associated with the first video frame.
+				This is what most writing programs do since the timecode of subsequent frames simply
+				follows from the first and videos usually do not have discontinuties.
+				When writing multiple video tracks and their associated timecode tracks you need to add
+				a reference to the timecode track to the video by calling `AddTrackReference()`.
+				If there is only a single timecode track that applies to all video tracks (or there is
+				only one video track) you should still associate the tracks but this is not strictly
+				necessary.
+			*/
+			bool bDropFrame = false;
+			bool bMax24Hours = false;
+			bool bAllowNegativeTimes = false;
+			// These values are defined to be signed integers!
+			int32 Timescale = 0;
+			int32 FrameDuration = 0;
+			int8 FramesPerSecond = 0;
+		};
+		struct FOpaqueData
+		{
+			/*
+				Arbitrary data can be written into the file as long as it conforms to the general
+				ISO/IEC 14496-12 sample structure.
+
+				This is useful if you want to add subtitle data where each of the formats like
+				`tx3g`, `wvtt` or `stpp` all use different handlers, media headers and data layout.
+				You can write any data of course provided that it does not contain offsets into
+				your data that depend on the position of any data written to the file, because
+				offsets cannot be adjusted by this writer for opaque data.
+
+				Since this writer only writes raw data in the first place, constructing the
+				boxes for any opaque data falls onto your application as well.
+			*/
+
+			// Values to put into the `tkhd` box.
+			int32 matrix[9] { 0x00010000,0,0,0,0x00010000,0,0,0,0x40000000 };
+			uint32 width = 0;
+			uint32 height = 0;
+			uint32 flags = 0;	// the track's `bEnabled` and `bInMovie` will overwrite what is set here.
+			int16 volume = 0;
+
+
+			// Type of handler to write to the `mdia`/`hdlr` box.
+			uint32 mdia_hdlr_Type = 0;
+			// Name of the handler to write to the `mdia`/`hdlr` box.
+			FString mdia_hdlr_Name;
+
+			// Type of media header box to write into the `minf` box (eg `nmhd`)
+			uint32 minf_MediaHeader_Type = 0;
+			// Raw data of the media header box (excluding the box size and box type)
+			// to write into the `minf` box.
+			// If this box is a full box, then the raw data needs to include the version and flags
+			// right before the payload.
+			TArray<uint8> minf_MediaHeader_RawBoxData;
+
+			// Raw `SampleEntry` (ISO/IEC 14496-12:2022 - 8.5.2.2 Syntax) data
+			// (excluding the box size and box type) to write into the `stsd` box.
+			// This sample entry data MUST include the base `SampleEntry` data members!
+			// The sample entry type is given by the track's `SampleEntryFormat` parameter (see below).
+			TArray<uint8> stsd_SampleEntry_RawBoxData;
+		};
 
 		/**
 		 * ISO/IEC 14496-12 Section 8.5.2.1 "format"
@@ -86,11 +158,11 @@ public:
 		/**
 		 * Track type specific values that need to be set.
 		 * The type this variant contains determines if the track contains
-		 * video or audio.
+		 * video, audio, timecode or other.
 		 * The values in these structures are written as-is and should
 		 * therefore be set correctly.
 		 */
-		TVariant<FVideo,FAudio> Properties;
+		TVariant<FVideo,FAudio,FTimecode,FOpaqueData> Properties;
 
 		/**
 		 * Additional pre-built boxes to be added to the sample type
@@ -177,6 +249,17 @@ public:
 	 */
 	virtual int32 AddTrack(const FTrackSpec& InTrackSpec) = 0;
 
+	/**
+	 * Adds a track reference.
+	 * This is usually only done for timecode tracks being referenced by one or more video tracks.
+	 * See ISO/IEC 14496-12:2022 section 8.3.3 Track reference box
+	 * Example:
+	 *  If you added a video track (track index 1) and a timecode track (track index 3) and you
+	 *  wish the video to reference the timecode track you call
+	 *  AddTrackReference(1, 3, MakeBoxAtom('t','m','c','d'))
+	 */
+	virtual bool AddTrackReference(int32 InTrackIndexToAddReferenceTo, int32 InTrackIndexBeingReferenced, uint32 InReferenceType) = 0;
+
 	/**
 	 * State value reported to the registered `FStatusDelegate`.
 	 */
