MP4Utilities: Added "WebOptimized" muxer mode which requires a 2nd pass

#rnx

#Synced-CL 49918411

diff --git a/Engine/Plugins/Media/MP4Utilities/Source/MP4Boxes/Private/MP4Boxes.cpp b/Engine/Plugins/Media/MP4Utilities/Source/MP4Boxes/Private/MP4Boxes.cpp
--- a/Engine/Plugins/Media/MP4Utilities/Source/MP4Boxes/Private/MP4Boxes.cpp
+++ b/Engine/Plugins/Media/MP4Utilities/Source/MP4Boxes/Private/MP4Boxes.cpp
@@ -473,7 +473,6 @@ namespace MP4Boxes
 		// Unsupported sample type. No need to continue parsing.
 		if (Parsed->SampleType == FMP4BoxSampleEntry::ESampleType::Unsupported)
 		{
-			check(!"TODO");
 			return;
 		}
 		uint32 EntryCount;
diff --git a/Engine/Plugins/Media/MP4Utilities/Source/MP4Muxer/Private/MP4Muxer.cpp b/Engine/Plugins/Media/MP4Utilities/Source/MP4Muxer/Private/MP4Muxer.cpp
--- a/Engine/Plugins/Media/MP4Utilities/Source/MP4Muxer/Private/MP4Muxer.cpp
+++ b/Engine/Plugins/Media/MP4Utilities/Source/MP4Muxer/Private/MP4Muxer.cpp
@@ -97,7 +97,7 @@ private:
 	}
 	bool SetWriteError()
 	{
-		return SetError(FString::Printf(TEXT("Could not write to file \"%s\"."), *Configuration.OutputFilename));
+		return SetError(FString::Printf(TEXT("Could not write to file \"%s\"."), *FilenameWritingNow));
 	}
 	bool SetUpdateError(const TCHAR* InWhat, const TCHAR* InBox)
 	{
@@ -112,12 +112,15 @@ private:
 	bool WriteAccumulatedSamples();
 	bool WriteFileEnd();
 	bool WriteUpdate(int64 InFilePosition, TConstArrayView<uint8> InData);
+	bool RewriteWebOptimizedFile();
+	void ClearUpdateBoxes();
 
 	FConfiguration Configuration;
 	TArray<FTrackSpec> TrackSpecs;
 	TMap<int32, FTrackReferences> TrackReferences;
 	FRequestTrackDataDelegate SampleRequestDelegate;
 	FStatusDelegate StatusDelegate;
+	FString FilenameWritingNow;
 	FString LastError;
 	bool bFailed = false;
 	bool bTerminate = false;
@@ -125,6 +128,7 @@ private:
 	FEvent* WorkerEvent = nullptr;
 	IFileHandle* FileHandle = nullptr;
 	int64 FileOffsetOfStandardMDAT = -1;
+	int64 FileEndOffsetOfStandardMDAT = -1;
 	uint32 FragmentSequenceNumber = 0;
 	bool bWroteMOOV = false;
 
@@ -197,6 +201,10 @@ bool FMP4RawMuxer::Configure(const FConfiguration& InConfiguration)
 	{
 		Configuration.InterleaveDuration = FTimespan::FromMilliseconds(500);
 	}
+	if (Configuration.MuxMode == FConfiguration::EMuxMode::WebOptimized && Configuration.TemporaryFilename.IsEmpty())
+	{
+		return SetError(FString::Printf(TEXT("WebOptimized mode requires a temporary file.")));
+	}
 
 	RewriteChangeDelegate.BindRaw(this, &FMP4RawMuxer::WriteUpdate);
 	return true;
@@ -322,7 +330,7 @@ bool FMP4RawMuxer::WriteFileStart()
 	FileLevelBoxes->AddChild(FTYP.Get());
 	MinimumMajorISOBrand = WrittenMajorISOBrand;
 	MP4MuxerBoxes::FBoxMDAT* MDAT = nullptr;
-	if (Configuration.MuxMode == IMP4RawMuxer::FConfiguration::EMuxMode::Standard)
+	if (Configuration.MuxMode != IMP4RawMuxer::FConfiguration::EMuxMode::Fragmented)
 	{
 		// We may need additional 8 bytes if the following `mdat` exceeds 4GiB,
 		// so we insert a `free` box with no additional zero bytes as this box size is 8 bytes.
@@ -599,7 +607,7 @@ bool FMP4RawMuxer::WriteFileStart()
 		MP4MuxerBoxes::FBoxSTTS* STTS = new MP4MuxerBoxes::FBoxSTTS;
 		TrackChunkSamples[TrackIdx].TrackMuxBoxes.STTS = STTS;
 		STBL->AddChild(STTS);
-		if (Configuration.MuxMode == IMP4RawMuxer::FConfiguration::EMuxMode::Standard)
+		if (Configuration.MuxMode != IMP4RawMuxer::FConfiguration::EMuxMode::Fragmented)
 		{
 			if (!TrackSpecs[TrackIdx].bIsAllKeyframes)
 			{
@@ -623,7 +631,7 @@ bool FMP4RawMuxer::WriteFileStart()
 		TrackChunkSamples[TrackIdx].TrackMuxBoxes.STCO = STCO;
 		STBL->AddChild(STCO);
 		// Sub-sample information? These go into the `stbl` for standard files and into the `traf` for fragmented files.
-		if (Configuration.MuxMode == IMP4RawMuxer::FConfiguration::EMuxMode::Standard && TrackSpecs[TrackIdx].SubSampleFlags.IsSet())
+		if (Configuration.MuxMode != IMP4RawMuxer::FConfiguration::EMuxMode::Fragmented && TrackSpecs[TrackIdx].SubSampleFlags.IsSet())
 		{
 			MP4MuxerBoxes::FBoxSUBS* SUBS = new MP4MuxerBoxes::FBoxSUBS;
 			TrackChunkSamples[TrackIdx].TrackMuxBoxes.SUBS = SUBS;
@@ -735,7 +743,7 @@ bool FMP4RawMuxer::WriteFileEnd()
 	}
 
 	// In standard mode the `moov` box has not been built yet. Do it and write it out.
-	if (Configuration.MuxMode == IMP4RawMuxer::FConfiguration::EMuxMode::Standard)
+	if (Configuration.MuxMode != IMP4RawMuxer::FConfiguration::EMuxMode::Fragmented)
 	{
 		// Check to see if we really need to write composition offsets for the tracks.
 		for(auto &tcs : TrackChunkSamples)
@@ -754,10 +762,52 @@ bool FMP4RawMuxer::WriteFileEnd()
 			}
 		}
 
+		// Check if any chunk offset box requires 64 bits. If one do then
+		// we set all to require 64 bits.
+		bool bUsing64BitChunkOffsets = false;
+		for(auto &tcs : TrackChunkSamples)
+		{
+			if (tcs.TrackMuxBoxes.STCO->Needs64Bits())
+			{
+				bUsing64BitChunkOffsets = true;
+				for(auto &tcs64 : TrackChunkSamples)
+				{
+					tcs64.TrackMuxBoxes.STCO->Force64Bits();
+				}
+				break;
+			}
+		}
+
+		// Remember the file position where we will write the `moov` box.
+		check(FileHandle->Tell() == FilePosAtEnd);
+		int64 FilePosWhereMoovStarts = FilePosAtEnd;
 
 		MP4MuxerBoxes::FBoxBuilderData builderData;
 		MOOV->BuildData(builderData);
 		MOOV->CalculateBoxSizes();
+		int64 MoovSize = MOOV->GetSize();
+		// Check if in WebOptimized mode the size of the file (mdat and moov (and others) combined)
+		// exceeds 4 GiB. If so we have to use 64 bit chunk offsets then because the `moov` will com
+		// before the `mdat` and thus all chunk offsets get larger and may not fit into 32 bits any more.
+		if (Configuration.MuxMode == IMP4RawMuxer::FConfiguration::EMuxMode::WebOptimized)
+		{
+			FileEndOffsetOfStandardMDAT = FilePosAtEnd;
+
+			int64 TotalFileSize = FilePosWhereMoovStarts + MoovSize;
+			const int32 kExtraSpace = 65536;
+			if (bUsing64BitChunkOffsets == false && (TotalFileSize + kExtraSpace) >= 0xffffffff)
+			{
+				for(auto &tcs64 : TrackChunkSamples)
+				{
+					tcs64.TrackMuxBoxes.STCO->Force64Bits();
+				}
+				// Rebuild the box tree.
+				MOOV->ClearBuiltBoxData();
+				MOOV->BuildData(builderData);
+				MOOV->CalculateBoxSizes();
+			}
+		}
+
 		bool bWriteOk = MOOV->VisitForWriting([&](MP4MuxerBoxes::FBoxBase* InBox)
 		{
 			InBox->SetAbsoluteFileOffset(FileHandle->Tell());
@@ -771,10 +821,9 @@ bool FMP4RawMuxer::WriteFileEnd()
 		// Get the final size of the `mdat` box.
 		const int64 SizeForMDATBox = FilePosAtEnd - FileOffsetOfStandardMDAT;
 		TArray<uint8> MDATBox;
-		int64 RewriteFileAt = FileOffsetOfStandardMDAT;
 		if (SizeForMDATBox > 0xffffffffU)
 		{
-			RewriteFileAt -= 8;
+			FileOffsetOfStandardMDAT -= 8;
 			MDATBox.SetNumZeroed(16);
 			uint32* bd = (uint32*)MDATBox.GetData();
 			*bd++ = MP4MuxerBoxes::EndianSwap((uint32)1U);
@@ -789,7 +838,7 @@ bool FMP4RawMuxer::WriteFileEnd()
 			*bd++ = MP4MuxerBoxes::EndianSwap((uint32)SizeForMDATBox);
 			*bd++ = MP4MuxerBoxes::EndianSwap(MP4MuxerBoxes::BoxAtom('m','d','a','t'));
 		}
-		bWriteOk = FileHandle->Seek(RewriteFileAt) && FileHandle->Write(MDATBox.GetData(), MDATBox.Num()) && FileHandle->Seek(FilePosAtEnd);
+		bWriteOk = FileHandle->Seek(FileOffsetOfStandardMDAT) && FileHandle->Write(MDATBox.GetData(), MDATBox.Num()) && FileHandle->Seek(FilePosAtEnd);
 		if (!bWriteOk)
 		{
 			return SetWriteError();
@@ -818,12 +867,6 @@ bool FMP4RawMuxer::WriteFileEnd()
 			return SetWriteError();
 		}
 	}
-
-	// Drop the boxes we needed for updating.
-	MOOV.Reset();
-	FTYP.Reset();
-	MVHD = nullptr;
-	MEHD = nullptr;
 	return true;
 }
 
@@ -930,7 +973,7 @@ bool FMP4RawMuxer::WriteAccumulatedSamples()
 		return true;
 	}
 	// Write standard or fragmented.
-	bool bOk = Configuration.MuxMode == IMP4RawMuxer::FConfiguration::EMuxMode::Standard ? WriteAccumulatedSamplesStandard() : WriteAccumulatedSamplesFragmented();
+	bool bOk = Configuration.MuxMode != IMP4RawMuxer::FConfiguration::EMuxMode::Fragmented ? WriteAccumulatedSamplesStandard() : WriteAccumulatedSamplesFragmented();
 	// Dump the samples that were written.
 	for(auto &tcs : TrackChunkSamples)
 	{
@@ -1279,13 +1322,125 @@ bool FMP4RawMuxer::WriteAccumulatedSamplesFragmented()
 }
 
 
+bool FMP4RawMuxer::RewriteWebOptimizedFile()
+{
+	FilenameWritingNow = Configuration.OutputFilename;
+	FileHandle = FPlatformFileManager::Get().GetPlatformFile().OpenWrite(*FilenameWritingNow, false, false);
+	if (!FileHandle)
+	{
+		return SetError(FString::Printf(TEXT("Could not open file \"%s\" for writing."), *FilenameWritingNow));
+	}
+
+	// Write the `ftyp` box
+	bool bWriteOk = FTYP->VisitForWriting([&](MP4MuxerBoxes::FBoxBase* InBox)
+	{
+		InBox->SetAbsoluteFileOffset(FileHandle->Tell());
+		return InBox->GetCompiledBoxData().Num() ? FileHandle->Write(InBox->GetCompiledBoxData().GetData(), InBox->GetCompiledBoxData().Num()) : true;
+	});
+	if (!bWriteOk)
+	{
+		return SetWriteError();
+	}
+
+	// Write the `moov` box
+	int64 MoovBoxSizeNow = MOOV->GetSize();
+	int64 MoovBoxFileOffset = FileHandle->Tell();
+	int64 MoovBoxFileEndOffset = MoovBoxFileOffset + MOOV->GetSize();
+
+	// Update all file chunk offsets
+	int64 OffsetToAdd = MoovBoxFileEndOffset - FileOffsetOfStandardMDAT;
+	for(auto &tcs64 : TrackChunkSamples)
+	{
+		if (!tcs64.TrackMuxBoxes.STCO->AddFileOffset(OffsetToAdd))
+		{
+			return SetError(FString::Printf(TEXT("Invalid chunk offset while optimizing file \"%s\"."), *FilenameWritingNow));
+		}
+	}
+	// Rebuild the box tree.
+	MP4MuxerBoxes::FBoxBuilderData builderData;
+	MOOV->ClearBuiltBoxData();
+	MOOV->BuildData(builderData);
+	MOOV->CalculateBoxSizes();
+	int64 NewMoovBoxSize = MOOV->GetSize();
+	if (NewMoovBoxSize != MoovBoxSizeNow)
+	{
+		return SetError(FString::Printf(TEXT("Size of `moov` box changed while optimizing file \"%s\"."), *FilenameWritingNow));
+	}
+	bWriteOk = MOOV->VisitForWriting([&](MP4MuxerBoxes::FBoxBase* InBox)
+	{
+		InBox->SetAbsoluteFileOffset(FileHandle->Tell());
+		return InBox->GetCompiledBoxData().Num() ? FileHandle->Write(InBox->GetCompiledBoxData().GetData(), InBox->GetCompiledBoxData().Num()) : true;
+	});
+	if (!bWriteOk)
+	{
+		return SetWriteError();
+	}
+
+	// Copy the `mdat`
+	uint32 kTempBufferSize = Configuration.TemporaryOptimizationBufferSize ? Configuration.TemporaryOptimizationBufferSize : 32 * 1024 * 1024;
+	uint8* TempBuffer = reinterpret_cast<uint8*>(FMemory::Malloc(kTempBufferSize));
+	if (!TempBuffer)
+	{
+		return SetError(FString::Printf(TEXT("Could not allocate temporary optimization buffer.")));
+	}
+
+	IFileHandle* SourceFileHandle = FPlatformFileManager::Get().GetPlatformFile().OpenRead(*Configuration.TemporaryFilename);
+	if (!SourceFileHandle)
+	{
+		FMemory::Free(TempBuffer);
+		return SetError(FString::Printf(TEXT("Failed to open temporary file \"%s\" for reading."), *Configuration.TemporaryFilename));
+	}
+	SourceFileHandle->Seek(FileOffsetOfStandardMDAT);
+	int64 SizeToRewrite = FileEndOffsetOfStandardMDAT - FileOffsetOfStandardMDAT;
+	int64 SizeToCopy = SizeToRewrite;
+	while(SizeToCopy > 0 && !bTerminate)
+	{
+		int64 Percentage = SizeToRewrite ? (10000 * (SizeToRewrite - SizeToCopy)) / SizeToRewrite : 0;
+		Configuration.OptimizeProgressDelegate.ExecuteIfBound((float)Percentage / 100.0);
+
+		int64 NumToRead = SizeToCopy > kTempBufferSize ? kTempBufferSize : SizeToCopy;
+		if (!SourceFileHandle->Read(TempBuffer, NumToRead))
+		{
+			delete SourceFileHandle;
+			FMemory::Free(TempBuffer);
+			return SetError(FString::Printf(TEXT("Failed to read from temporary file \"%s\" while optimizing."), *Configuration.TemporaryFilename));
+		}
+		if (!FileHandle->Write(TempBuffer, NumToRead))
+		{
+			delete SourceFileHandle;
+			FMemory::Free(TempBuffer);
+			return SetError(FString::Printf(TEXT("Failed to write to file \"%s\" while optimizing."), *FilenameWritingNow));
+		}
+		SizeToCopy -= NumToRead;
+	}
+	delete SourceFileHandle;
+	FMemory::Free(TempBuffer);
+	if (!bTerminate)
+	{
+		Configuration.OptimizeProgressDelegate.ExecuteIfBound((float)100.0f);
+		// Delete the temp file.
+		FPlatformFileManager::Get().GetPlatformFile().DeleteFile(*Configuration.TemporaryFilename);
+	}
+	return true;
+}
+
+void FMP4RawMuxer::ClearUpdateBoxes()
+{
+	// Drop the boxes we needed for updating.
+	MOOV.Reset();
+	FTYP.Reset();
+	MVHD = nullptr;
+	MEHD = nullptr;
+}
+
 uint32 FMP4RawMuxer::Run()
 {
-	// Open the file
-	FileHandle = FPlatformFileManager::Get().GetPlatformFile().OpenWrite(*Configuration.OutputFilename, false, false);
+	// Open the file. Temporary one for WebOptimized mode, final output for regular or fragmented mode.
+	FilenameWritingNow = Configuration.MuxMode == FConfiguration::EMuxMode::WebOptimized ? Configuration.TemporaryFilename : Configuration.OutputFilename;
+	FileHandle = FPlatformFileManager::Get().GetPlatformFile().OpenWrite(*FilenameWritingNow, false, false);
 	if (!FileHandle)
 	{
-		SetError(FString::Printf(TEXT("Could not open file \"%s\" for writing."), *Configuration.OutputFilename));
+		SetError(FString::Printf(TEXT("Could not open file \"%s\" for writing."), *FilenameWritingNow));
 		StatusDelegate.ExecuteIfBound(EStatus::Failed);
 	}
 	else if (!WriteFileStart())
@@ -1334,15 +1489,25 @@ uint32 FMP4RawMuxer::Run()
 			{
 				bool bOk = WriteFileEnd();
 				CloseFile();
-				if (bOk)
+				if (bOk && Configuration.MuxMode == FConfiguration::EMuxMode::WebOptimized)
 				{
-					StatusDelegate.ExecuteIfBound(EStatus::Finished);
+					// Rewrite optimized file.
+					bOk = RewriteWebOptimizedFile();
+					CloseFile();
 				}
-				else
+				ClearUpdateBoxes();
+				if (!bTerminate)
 				{
-					StatusDelegate.ExecuteIfBound(EStatus::Failed);
+					if (bOk)
+					{
+						StatusDelegate.ExecuteIfBound(EStatus::Finished);
+					}
+					else
+					{
+						StatusDelegate.ExecuteIfBound(EStatus::Failed);
+					}
+					bIsDone = true;
 				}
-				bIsDone = true;
 			}
 		}
 	}
diff --git a/Engine/Plugins/Media/MP4Utilities/Source/MP4Muxer/Private/MP4MuxerBoxes.cpp b/Engine/Plugins/Media/MP4Utilities/Source/MP4Muxer/Private/MP4MuxerBoxes.cpp
--- a/Engine/Plugins/Media/MP4Utilities/Source/MP4Muxer/Private/MP4MuxerBoxes.cpp
+++ b/Engine/Plugins/Media/MP4Utilities/Source/MP4Muxer/Private/MP4MuxerBoxes.cpp
@@ -211,6 +211,19 @@ FBoxBase::~FBoxBase()
 }
 
 
+void FBoxBase::ClearBuiltBoxData()
+{
+	for(int32 i=0; i<Children.Num(); ++i)
+	{
+		Children[i]->ClearBuiltBoxData();
+	}
+	Size = 0;
+	CompiledBoxData.Empty();
+	CompiledBoxData_Offset_Size = 0;
+	AbsoluteFileOffset = -1;
+	bWasBuilt = false;
+}
+
 bool FBoxBase::BuildData(const FBoxBuilderData& InOutBuilderData)
 {
 	// Build children recursively first.
@@ -881,6 +894,20 @@ bool FBoxSTSZ::BuildData(const FBoxBuilderData& InOutBuilderData)
 }
 
 
+bool FBoxSTCO::AddFileOffset(int64 InOffsetToAdd)
+{
+	for(int32 i=0, iMax=ChunkOffsets.Num(); i<iMax; ++i)
+	{
+		int64 Offset = ChunkOffsets[i] + InOffsetToAdd;
+		if (Offset < 0 || (!bNeeds64Bits && Offset > 0xffffffff))
+		{
+			return false;
+		}
+		ChunkOffsets[i] = Offset;
+	}
+	return true;
+}
+
 bool FBoxSTCO::BuildData(const FBoxBuilderData& InOutBuilderData)
 {
 	if (bNeeds64Bits)
diff --git a/Engine/Plugins/Media/MP4Utilities/Source/MP4Muxer/Private/MP4MuxerBoxes.h b/Engine/Plugins/Media/MP4Utilities/Source/MP4Muxer/Private/MP4MuxerBoxes.h
--- a/Engine/Plugins/Media/MP4Utilities/Source/MP4Muxer/Private/MP4MuxerBoxes.h
+++ b/Engine/Plugins/Media/MP4Utilities/Source/MP4Muxer/Private/MP4MuxerBoxes.h
@@ -33,6 +33,7 @@ public:
 	FBoxBase(uint32 InType) : Type(InType)
 	{ }
 	virtual ~FBoxBase();
+	virtual void ClearBuiltBoxData();
 	virtual bool BuildData(const FBoxBuilderData& InOutBuilderData);
 	virtual void CalculateBoxSizes();
 	void AddChild(FBoxBase* InChild)
@@ -73,6 +74,8 @@ public:
 		}
 		return true;
 	}
+	int64 GetSize() const
+	{ return Size; }
 	const TArray<uint8>& GetCompiledBoxData() const
 	{ return CompiledBoxData; }
 	// Gets set while writing
@@ -98,6 +101,7 @@ protected:
 	TArray<uint8> CompiledBoxData;
 	uint32 CompiledBoxData_Offset_Size = 0;
 	int64 AbsoluteFileOffset = -1;
+	bool bWasBuilt = false;
 };
 
 class FBoxFull : public FBoxBase
@@ -137,7 +141,6 @@ protected:
 	uint32 MinorVersion = 0;
 	TArray<uint32> CompatibleBrands;
 	int64 MajorBrandPatchOffset = -1;
-	bool bWasBuilt = false;
 };
 
 
@@ -209,7 +212,6 @@ protected:
 	uint32 NextTrackID = 0;
 	int64 TimescalePatchOffset = -1;
 	int64 DurationPatchOffset = -1;
-	bool bWasBuilt = false;
 };
 
 
@@ -282,7 +284,6 @@ protected:
 	uint16 Layer = 0;
 	uint16 AlternateGroup = 0;
 	int64 DurationPatchOffset = -1;
-	bool bWasBuilt = false;
 };
 
 
@@ -329,7 +330,6 @@ protected:
 	uint32 Timescale = 0;
 	uint8 LanguageCode[3] {};
 	int64 DurationPatchOffset = -1;
-	bool bWasBuilt = false;
 };
 
 
@@ -699,6 +699,13 @@ public:
 	}
 	uint32 GetNumberOfChunkOffsets() const
 	{ return ChunkOffsets.Num(); }
+	const TChunkedArray<int64, 32>& GetChunkOffsets() const
+	{ return ChunkOffsets; }
+	bool Needs64Bits() const
+	{ return bNeeds64Bits; }
+	void Force64Bits()
+	{ bNeeds64Bits = true; }
+	bool AddFileOffset(int64 InOffsetToAdd);
 	bool BuildData(const FBoxBuilderData& InOutBuilderData) override;
 protected:
 	TChunkedArray<int64, 32> ChunkOffsets;
@@ -730,7 +737,6 @@ public:
 protected:
 	uint64 FragmentDuration = 0;
 	int64 FragmentDurationPatchOffset = -1;
-	bool bWasBuilt = false;
 };
 
 
@@ -829,7 +835,6 @@ protected:
 	TOptional<uint32> default_sample_size;
 	TOptional<uint32> default_sample_flags;
 	int64 BaseDataPatchOffset = -1;
-	bool bWasBuilt = false;
 };
 
 
@@ -994,7 +999,6 @@ protected:
 	uint64 EditDuration = 0;
 	int64 MediaTime = 0;
 	int64 EditDurationPatchOffset = -1;
-	bool bWasBuilt = false;
 };
 
 
diff --git a/Engine/Plugins/Media/MP4Utilities/Source/MP4Muxer/Public/MP4Muxer.h b/Engine/Plugins/Media/MP4Utilities/Source/MP4Muxer/Public/MP4Muxer.h
--- a/Engine/Plugins/Media/MP4Utilities/Source/MP4Muxer/Public/MP4Muxer.h
+++ b/Engine/Plugins/Media/MP4Utilities/Source/MP4Muxer/Public/MP4Muxer.h
@@ -209,7 +209,11 @@ public:
 	{
 		enum class EMuxMode
 		{
+			// Standard file, `moov` box at the end.
 			Standard,
+			// Optimized file, `moov` box in front. Requires a temporary file and a 2nd cleanup pass.
+			WebOptimized,
+			// Fragmented file
 			Fragmented,
 		};
 
@@ -218,6 +222,24 @@ public:
 		 */
 		FString OutputFilename;
 
+		/**
+		 * Filename for the temporary file in `WebOptimized` mode.
+		 * It will be deleted after optimization unless an error occurred
+		 * or you have aborted muxing by calling `StopAndClose()`.
+		 */
+		FString TemporaryFilename;
+
+		/**
+		 * Size of the temporary buffer used to read-write data to and from
+		 * during the optimization phase in `WebOptimized` mode.
+		 * Data of the unoptimized file is read into this buffer and written
+		 * into the optimized file from here. The larger the buffer the
+		 * fewer read/writes are necessary.
+		 * If kept at zero a suitable buffer size is chosen automatically.
+		 * This is only allocated in `WebOptimized` mode.
+		 */
+		uint32 TemporaryOptimizationBufferSize = 0;
+
 		/**
 		 * Whether to write a standard mp4 (`moov` box will be placed at the end) or
 		 * a fragmented mp4.
@@ -235,6 +257,19 @@ public:
 		 * needs to collect and keep, requiring more memory.
 		 */
 		FTimespan InterleaveDuration;
+
+
+		/**
+		 * Optional progress delegate to invoke periodically during the optimization
+		 * phase in `WebOptimized` mode.
+		 * There is no progress delegate for `Standard` or `Fragmented` mode since
+		 * the muxer does no know how many samples will be written.
+		 * If you want to display progress you can do this during the muxing process
+		 * by tracking which sample number (out of however many, which you should know)
+		 * is currently being written into the file.
+		 */
+		DECLARE_DELEGATE_OneParam(FOptimizeProgressDelegate, float /*PrecentComplete*/);
+		FOptimizeProgressDelegate OptimizeProgressDelegate;
 	};
 
 	/**
@@ -265,7 +300,9 @@ public:
 	 */
 	enum class EStatus
 	{
+		// Failure.
 		Failed,
+		// Multiplexing completed successfully.
 		Finished
 	};
 	DECLARE_DELEGATE_OneParam(FStatusDelegate, EStatus /*Status*/);
@@ -317,6 +354,12 @@ public:
 	 *   message by calling `GetLastError()`.
 	 *   When muxing is complete and the file has been successully written the `FStatusDelegate`
 	 *   will be invoked with an `EStatus::Finished`.
+	 *
+	 *   If you abort muxing by calling `StopAndClose()` the `FStatusDelegate` will not be invoked.
+	 *
+	 *   If an error occurs the partially written file will be left in place and not be
+	 *   deleted. This is true for the temporary file in `WebOptimized` mode as well,
+	 *   which is only deleted if the optimization step has finished successfully.
 	 */
 	virtual bool Start(FRequestTrackDataDelegate InSampleRequestDelegate, FStatusDelegate InStatusDelegate) = 0;
 
@@ -326,7 +369,6 @@ public:
 	 * during the muxing process.
 	 * All members should be set correctly as no plausibility checks
 	 * are performed except for
-	 *  - `Duration` not being zero or less
 	 *  - `Data.Num()` being less than 2^32. A single sample cannot
 	 *     be larger than 4 GiB in an ISO/IEC 14496-12:2022 file.
 	 */
@@ -351,6 +393,7 @@ public:
 		 * Presentation duration of this sample, given in the tracks timescale units.
 		 * Must not be zero unless a constant duration has been set for the track
 		 * which must not be zero.
+		 * The value may be zero for the last sample although this is not recommended.
 		 * This value is the most important value as it is written to the `stts`
 		 * box directly and their relative sum determines the decode time (DTS).
 		 */
@@ -417,7 +460,6 @@ public:
 	 * Returns a human readable error message if muxing failed.
 	 */
 	virtual FString GetLastError() = 0;
-
 };
 
 #undef UE_API
