MP4Utilities: fixed handling of encrypted tracks that have no senc box; only reading well-known root boxes

#rnx

#Synced-CL 51366553

diff --git a/Engine/Plugins/Media/MP4Utilities/Source/MP4Boxes/Private/MP4Track.cpp b/Engine/Plugins/Media/MP4Utilities/Source/MP4Boxes/Private/MP4Track.cpp
--- a/Engine/Plugins/Media/MP4Utilities/Source/MP4Boxes/Private/MP4Track.cpp
+++ b/Engine/Plugins/Media/MP4Utilities/Source/MP4Boxes/Private/MP4Track.cpp
@@ -586,7 +586,7 @@ namespace MP4Boxes
 			// Are there subsample boxes? If yes, we use the first one only in the iterator for now.
 			It->subsIt.SetBox(SubsBoxes.Num() ? SubsBoxes[0] : nullptr, Convs.NumMoovSamples);
 			// Encryption?
-			It->sencIt.SetBox(SencBox, Convs.NumMoovSamples);
+			It->sencIt.SetBox(TencBox, SencBox, Convs.NumMoovSamples);
 			// Do we have a `rap ` group?
 			TSharedPtr<FMP4BoxSGPD, ESPMode::ThreadSafe>* rapSGPD = SgpdBoxes.FindByPredicate([](const TSharedPtr<FMP4BoxSGPD, ESPMode::ThreadSafe>& e){ return e->GetGroupingType() == MP4Utilities::MakeBoxAtom('r','a','p',' ');});
 			TSharedPtr<FMP4BoxSBGP, ESPMode::ThreadSafe>* rapSBGP = SbgpBoxes.FindByPredicate([](const TSharedPtr<FMP4BoxSBGP, ESPMode::ThreadSafe>& e){ return e->GetGroupingType() == MP4Utilities::MakeBoxAtom('r','a','p',' ');});
@@ -1051,7 +1051,7 @@ namespace MP4Boxes
 				subsIt.SetToSampleNumber(SampleNumber);
 
 				// Likewise there can only be a single `senc` box.
-				sencIt.SetBox(Traf->SencBox, Traf->NumSamples);
+				sencIt.SetBox(FragmentInfo->TencBox, Traf->SencBox, Traf->NumSamples);
 				sencIt.SetToSampleNumber(SampleNumber);
 			}
 		}
@@ -1242,7 +1242,7 @@ namespace MP4Boxes
 			subsIt.SetToSampleNumber(SampleNumber);
 
 			// Encryption?
-			sencIt.SetBox(Traf->SencBox, Traf->NumSamples);
+			sencIt.SetBox(FragmentInfo->TencBox, Traf->SencBox, Traf->NumSamples);
 			sencIt.SetToSampleNumber(SampleNumber);
 
 			Update(Trun);
@@ -1310,7 +1310,7 @@ namespace MP4Boxes
 			subsIt.SetToSampleNumber(SampleNumber);
 
 			// Encryption?
-			sencIt.SetBox(Traf->SencBox, Traf->NumSamples);
+			sencIt.SetBox(FragmentInfo->TencBox, Traf->SencBox, Traf->NumSamples);
 			sencIt.SetToSampleNumber(SampleNumber);
 
 			Update(Trun);
diff --git a/Engine/Plugins/Media/MP4Utilities/Source/MP4Boxes/Public/MP4BoxIterators.h b/Engine/Plugins/Media/MP4Utilities/Source/MP4Boxes/Public/MP4BoxIterators.h
--- a/Engine/Plugins/Media/MP4Utilities/Source/MP4Boxes/Public/MP4BoxIterators.h
+++ b/Engine/Plugins/Media/MP4Utilities/Source/MP4Boxes/Public/MP4BoxIterators.h
@@ -929,7 +929,7 @@ namespace MP4Boxes
 
 	struct FSENCBoxIterator
 	{
-		void SetBox(const TSharedPtr<FMP4BoxSENC, ESPMode::ThreadSafe>& InSencBox, uint32 InMaxSamples)
+		void SetBox(const TSharedPtr<FMP4BoxTENC, ESPMode::ThreadSafe>& InTencBox, const TSharedPtr<FMP4BoxSENC, ESPMode::ThreadSafe>& InSencBox, uint32 InMaxSamples)
 		{
 			NumTotalSamples = InMaxSamples;
 			CurrentSampleNum = 0;
@@ -938,6 +938,10 @@ namespace MP4Boxes
 			{
 				NumBoxSamples = SencBox->GetEntries().Num();
 			}
+			if (InTencBox.IsValid())
+			{
+				DefaultEntry.IV = InTencBox->GetDefaultConstantIV();
+			}
 		}
 
 		void SetToSampleNumber(uint32 InSampleNum)
@@ -973,12 +977,12 @@ namespace MP4Boxes
 
 		const FMP4BoxSENC::FEntry& GetCurrentEntry() const
 		{
-			return CurrentSampleNum < NumBoxSamples ? SencBox->GetEntries()[CurrentSampleNum] : EmptyEntry;
+			return CurrentSampleNum < NumBoxSamples ? SencBox->GetEntries()[CurrentSampleNum] : DefaultEntry;
 		}
 
 	private:
 		TSharedPtr<FMP4BoxSENC, ESPMode::ThreadSafe> SencBox;
-		FMP4BoxSENC::FEntry EmptyEntry;
+		FMP4BoxSENC::FEntry DefaultEntry;
 		uint32 NumTotalSamples = 0;
 		uint32 NumBoxSamples = 0;
 		uint32 CurrentSampleNum = 0;
diff --git a/Engine/Plugins/Media/MP4Utilities/Source/MP4Utilities/Private/MP4Utilities.cpp b/Engine/Plugins/Media/MP4Utilities/Source/MP4Utilities/Private/MP4Utilities.cpp
--- a/Engine/Plugins/Media/MP4Utilities/Source/MP4Utilities/Private/MP4Utilities.cpp
+++ b/Engine/Plugins/Media/MP4Utilities/Source/MP4Utilities/Private/MP4Utilities.cpp
@@ -1,7 +1,7 @@
 // Copyright Epic Games, Inc. All Rights Reserved.
 
 #include "MP4Utilities.h"
-
+#include "MP4UtilitiesModule.h"
 
 namespace MP4Utilities
 {
@@ -126,7 +126,8 @@ namespace MP4Utilities
 		Buf.AddUninitialized(InNumBytes);
 		if (ReadBytes(Buf.GetData(), InNumBytes))
 		{
-			check(!"TODO");
+		// TODO
+			unimplemented();
 	/*
 			FUTF8ToTCHAR cnv((const ANSICHAR*)Buf.GetData(), InNumBytes);
 			OutString = FString(cnv.Length(), cnv.Get());
@@ -304,11 +305,36 @@ namespace MP4Utilities
 
 	bool FMP4BoxLocator::LocateAndReadRootBoxes(TArray<TSharedPtr<FMP4BoxData, ESPMode::ThreadSafe>>& OutBoxInfos, const TSharedPtr<IMP4DataReaderBase, ESPMode::ThreadSafe>& InDataReader, const TArray<uint32>& InFirstBoxes, const TArray<uint32>& InStopAfterBoxes, bool bStopWithMDAT, const TArray<uint32>& InReadDataOfBoxes, IMP4DataReaderBase::FCancellationCheckDelegate InCheckCancellationDelegate)
 	{
+		static const TArray<uint32> kWellKnownRootBoxes {
+			MakeBoxAtom('f','t','y','p'), MakeBoxAtom('o','t','y','p'), MakeBoxAtom('p','d','i','n'),
+			MakeBoxAtom('m','o','o','v'), MakeBoxAtom('m','o','o','f'), MakeBoxAtom('m','f','r','a'),
+			MakeBoxAtom('m','d','a','t'), MakeBoxAtom('f','r','e','e'), MakeBoxAtom('s','k','i','p'),
+			MakeBoxAtom('i','m','d','a'), MakeBoxAtom('m','e','t','a'), MakeBoxAtom('s','t','y','p'),
+			MakeBoxAtom('s','i','d','x'), MakeBoxAtom('s','s','i','x'), MakeBoxAtom('p','r','f','t'),
+			MakeBoxAtom('u','u','i','d')
+		};
+
 		// We NEVER want to read the `mdat` box here!
 		if (!InDataReader.IsValid() || InReadDataOfBoxes.Contains(MakeBoxAtom('m','d','a','t')))
 		{
 			return false;
 		}
+#if !UE_BUILD_SHIPPING
+		auto WarnNotWellKnown = [&](const TArray<uint32>& InBoxes) -> void
+		{
+			for(int32 i=0; i<InBoxes.Num(); ++i)
+			{
+				if (!kWellKnownRootBoxes.Contains(InBoxes[i]))
+				{
+					UE_LOG(LogMP4Utilities, Warning, TEXT("Box `%s` is not a well-known root box"), *GetPrintableBoxAtom(InBoxes[i]));
+				}
+			}
+		};
+		WarnNotWellKnown(InFirstBoxes);
+		WarnNotWellKnown(InStopAfterBoxes);
+		WarnNotWellKnown(InReadDataOfBoxes);
+#endif
+
 		CurrentOffset = InDataReader->GetCurrentFileOffset();
 
 		#define CHECK_READ(NumReq) \
@@ -343,6 +369,7 @@ namespace MP4Utilities
 	#if !UE_BUILD_SHIPPING
 			bi->Name[0] = BoxSizeAndType.As8[4]; bi->Name[1] = BoxSizeAndType.As8[5]; bi->Name[2] = BoxSizeAndType.As8[6]; bi->Name[3] = BoxSizeAndType.As8[7]; bi->Name[4] = 0;
 	#endif
+			//UE_LOG(LogMP4Utilities, Warning, TEXT("[%p]: box `%s` (%8x); %lld @ %lld"), this, *GetPrintableBoxAtom(bi->Type), bi->Type, (long long int) bi->Size, (long long int) bi->Offset);
 
 			// After having read the first few bytes we should now know the overall filesize.
 			if (BoxNum == 0)
@@ -382,14 +409,32 @@ namespace MP4Utilities
 			// Shall we read this box?
 			if (bi->Type != MakeBoxAtom('m','d','a','t'))
 			{
-				if (InReadDataOfBoxes.IsEmpty() || InReadDataOfBoxes.Contains(bi->Type))
+				if (kWellKnownRootBoxes.Contains(bi->Type))
 				{
-					bi->DataBuffer.SetNumUninitialized(bi->Size-BoxInternalOffset);
-					bi->Data = MakeConstArrayView(bi->DataBuffer);
-					NumRead = InDataReader->ReadData(bi->DataBuffer.GetData(), bi->Size-BoxInternalOffset, CurrentOffset+BoxInternalOffset, InCheckCancellationDelegate);
+					if (InReadDataOfBoxes.IsEmpty() || InReadDataOfBoxes.Contains(bi->Type))
+					{
+						// Size safety check.
+						const uint64 kMaxSafeBoxSize = 1024 * 1024 * 128;
+						if ((uint64)bi->Size <= kMaxSafeBoxSize)
+						{
+							bi->DataBuffer.SetNumUninitialized(bi->Size-BoxInternalOffset);
+							bi->Data = MakeConstArrayView(bi->DataBuffer);
+							NumRead = InDataReader->ReadData(bi->DataBuffer.GetData(), bi->Size-BoxInternalOffset, CurrentOffset+BoxInternalOffset, InCheckCancellationDelegate);
+						}
+						else
+						{
+							UE_LOG(LogMP4Utilities, Warning, TEXT("Skipping over box `%s` because it is suspiciously large (%lld bytes)"), *GetPrintableBoxAtom(bi->Type), (long long int) bi->Size);
+							NumRead = InDataReader->ReadData(nullptr, bi->Size-BoxInternalOffset, CurrentOffset+BoxInternalOffset, InCheckCancellationDelegate);
+						}
+					}
+					else
+					{
+						NumRead = InDataReader->ReadData(nullptr, bi->Size-BoxInternalOffset, CurrentOffset+BoxInternalOffset, InCheckCancellationDelegate);
+					}
 				}
 				else
 				{
+					UE_LOG(LogMP4Utilities, Log, TEXT("Skipping over box `%s` because it is not well known"), *GetPrintableBoxAtom(bi->Type));
 					NumRead = InDataReader->ReadData(nullptr, bi->Size-BoxInternalOffset, CurrentOffset+BoxInternalOffset, InCheckCancellationDelegate);
 				}
 			}
